Landscape: First draft of the experimental landscape patch plugin, with a working circle height patch.

diff --git a/Engine/Plugins/Experimental/LandscapePatch/LandscapePatch.uplugin b/Engine/Plugins/Experimental/LandscapePatch/LandscapePatch.uplugin
new file mode 100644
--- /dev/null
+++ b/Engine/Plugins/Experimental/LandscapePatch/LandscapePatch.uplugin
@@ -0,0 +1,24 @@
+{
+	"FileVersion": 3,
+	"Version": 1,
+	"VersionName": "0.1",
+	"FriendlyName": "Landscape Patch",
+	"Description": "Support for adding landscape patches- components that can be attached to meshes to affect the landscape as the mesh is repositioned.",
+	"Category": "Other",
+	"CreatedBy": "Epic Games, Inc.",
+	"CreatedByURL": "http://epicgames.com",
+	"DocsURL": "",
+	"MarketplaceURL": "",
+	"SupportURL": "",
+	"CanContainContent": true,
+	"IsBetaVersion": true,
+	"IsExperimentalVersion": true,
+	"Installed": false,
+	"Modules": [
+		{
+			"Name": "LandscapePatch",
+			"Type": "Runtime",
+			"LoadingPhase": "PostConfigInit"
+		}
+	]
+}
\ No newline at end of file
diff --git a/Engine/Plugins/Experimental/LandscapePatch/Shaders/Private/LandscapeCircleHeightPatchPS.usf b/Engine/Plugins/Experimental/LandscapePatch/Shaders/Private/LandscapeCircleHeightPatchPS.usf
new file mode 100644
--- /dev/null
+++ b/Engine/Plugins/Experimental/LandscapePatch/Shaders/Private/LandscapeCircleHeightPatchPS.usf
@@ -0,0 +1,35 @@
+// Copyright Epic Games, Inc. All Rights Reserved.
+
+#pragma once
+
+#include "/Engine/Private/Common.ush"
+#include "/Engine/Private/Landscape/LandscapeCommon.ush"
+
+float3 InCenter;
+float InRadius;
+float InFalloff;
+Texture2D<float4> InSourceHeightmap;
+
+void ApplyLandscapeCircleHeightPatch(in float4 SVPos : SV_POSITION, out float2 OutColor : SV_Target0)
+{
+	const float KINDA_SMALL_NUMBER = 0.0001;
+	const float ClampedFalloff = max(InFalloff, KINDA_SMALL_NUMBER);
+	
+	int2 TextureCoordinates = floor(SVPos.xy);
+	float4 CurrentPackedHeight = InSourceHeightmap.Load(int3(TextureCoordinates, 0));
+	float CurrentHeight = UnpackHeight(CurrentPackedHeight.xy);
+	
+	float PatchHeight = InCenter.z;
+	
+	float Distance = distance(TextureCoordinates, InCenter.xy);
+	
+	float Alpha = 1.0;
+	if (Distance > InRadius)
+	{
+		Alpha = 1 - min((Distance - InRadius) / ClampedFalloff, 1.0);
+	}
+	
+	float NewHeight = lerp(CurrentHeight, PatchHeight, Alpha);
+	
+	OutColor = PackHeight(NewHeight);
+}
\ No newline at end of file
diff --git a/Engine/Plugins/Experimental/LandscapePatch/Source/LandscapePatch/LandscapePatch.Build.cs b/Engine/Plugins/Experimental/LandscapePatch/Source/LandscapePatch/LandscapePatch.Build.cs
new file mode 100644
--- /dev/null
+++ b/Engine/Plugins/Experimental/LandscapePatch/Source/LandscapePatch/LandscapePatch.Build.cs
@@ -0,0 +1,45 @@
+// Copyright Epic Games, Inc. All Rights Reserved.
+
+using UnrealBuildTool;
+
+public class LandscapePatch : ModuleRules
+{
+	public LandscapePatch(ReadOnlyTargetRules Target) : base(Target)
+	{
+		PCHUsage = ModuleRules.PCHUsageMode.UseExplicitOrSharedPCHs;
+		
+		PublicIncludePaths.AddRange(
+			new string[] {
+				// ... add public include paths required here ...
+			}
+			);
+				
+		
+		PrivateIncludePaths.AddRange(
+			new string[] {
+				// ... add other private include paths required here ...
+			}
+			);
+			
+		
+		PublicDependencyModuleNames.AddRange(
+			new string[]
+			{
+				"Core",
+				"CoreUObject",
+				"Landscape",
+				"Engine",
+			}
+			);
+			
+		
+		PrivateDependencyModuleNames.AddRange(
+			new string[]
+			{
+				"RenderCore",
+				"RHI",
+				"Projects", // IPluginManager
+			}
+			);
+	}
+}
diff --git a/Engine/Plugins/Experimental/LandscapePatch/Source/LandscapePatch/Private/LandscapeCircleHeightPatch.cpp b/Engine/Plugins/Experimental/LandscapePatch/Source/LandscapePatch/Private/LandscapeCircleHeightPatch.cpp
new file mode 100644
--- /dev/null
+++ b/Engine/Plugins/Experimental/LandscapePatch/Source/LandscapePatch/Private/LandscapeCircleHeightPatch.cpp
@@ -0,0 +1,110 @@
+// Copyright Epic Games, Inc. All Rights Reserved.
+
+#include "LandscapeCircleHeightPatch.h"
+
+#include "Engine/TextureRenderTarget2D.h"
+#include "LandscapeCircleHeightPatchPS.h"
+#include "LandscapePatchManager.h"
+#include "RenderGraph.h" // RDG_EVENT_NAME
+
+#define LOCTEXT_NAMESPACE "LandscapeCircleHeightPatch"
+
+void ULandscapeCircleHeightPatch::OnComponentCreated()
+{
+	Super::OnComponentCreated();
+
+	// If we haven't been made from a copy, initialize the radius and transform of the patch
+	// based on our parent.
+	if (!bWasCopy)
+	{
+		AActor* ParentActor = GetAttachParentActor();
+		if (ParentActor)
+		{
+			FVector Origin, BoxExtent;
+			GetAttachParentActor()->GetActorBounds(false, Origin, BoxExtent);
+
+			// Place the component at the bottom of the bounding box.
+			Origin.Z -= BoxExtent.Z;
+			SetWorldLocation(Origin);
+
+			Radius = FMath::Sqrt(BoxExtent.X * BoxExtent.X + BoxExtent.Y * BoxExtent.Y);
+			Falloff = Radius / 2;
+		}
+	}
+}
+
+void ULandscapeCircleHeightPatch::Initialize_Native(const FTransform & InLandscapeTransform,
+	const FIntPoint& InLandscapeSize,
+	const FIntPoint& InLandscapeRenderTargetSize)
+{
+}
+
+UTextureRenderTarget2D* ULandscapeCircleHeightPatch::Render_Native(bool bIsHeightmap,
+	UTextureRenderTarget2D* InCombinedResult,
+	const FName& InWeightmapLayerName)
+{
+	if (!ensure(PatchManager.IsValid()))
+	{
+		return InCombinedResult;
+	}
+
+	// Circle height patch doesn't affect weight maps
+	if (!bIsHeightmap)
+	{
+		return InCombinedResult;
+	}
+
+	FTransform HeightmapCoordsToWorld = PatchManager->GetHeightmapCoordsToWorld();
+
+	double ToHeightmapRadiusScale = GetComponentTransform().GetScale3D().X / HeightmapCoordsToWorld.GetScale3D().X;
+
+	FVector3d CircleCenterWorld = GetComponentTransform().GetTranslation();
+	FVector3d CenterInHeightmapCoordinates = HeightmapCoordsToWorld.InverseTransformPosition(CircleCenterWorld);
+
+	float RadiusAdjustment = bExclusiveRadius ? 0 : 1;
+
+	float HeightmapRadius = Radius * ToHeightmapRadiusScale + RadiusAdjustment;
+	float HeightmapFalloff = Falloff * ToHeightmapRadiusScale + RadiusAdjustment;
+
+	// Go ahead an pack everything into a copy of the param struct so we don't have to capture everything
+	// individually in the lambda below.
+	FLandscapeCircleHeightPatchPS::FParameters ShaderParamsToCopy;
+	ShaderParamsToCopy.InCenter = (FVector3f)CenterInHeightmapCoordinates;
+	ShaderParamsToCopy.InRadius = HeightmapRadius;
+	ShaderParamsToCopy.InFalloff = HeightmapFalloff;
+
+	ENQUEUE_RENDER_COMMAND(LandscapeCircleHeightPatch)([InCombinedResult, ShaderParamsToCopy](FRHICommandListImmediate& RHICmdList)
+	{
+		TRACE_CPUPROFILER_EVENT_SCOPE(LandscapeCircleHeightPatch);
+
+		FRDGBuilder GraphBuilder(RHICmdList, RDG_EVENT_NAME("ApplyLandmassBoxHeightPatch"));
+		
+		TRefCountPtr<IPooledRenderTarget> RenderTarget = CreateRenderTarget(InCombinedResult->GetResource()->GetTexture2DRHI(), TEXT("LandscapeCircleHeightPatchOutput"));
+		FRDGTextureRef DestinationTexture = GraphBuilder.RegisterExternalTexture(RenderTarget);
+
+		// Make a copy of our heightmap input
+		FRDGTextureRef InputCopy = GraphBuilder.CreateTexture(DestinationTexture->Desc, TEXT("LandscapeCircleHeightPatchInputCopy"));
+
+		FRHICopyTextureInfo CopyTextureInfo;
+		CopyTextureInfo.NumMips = 1;
+		CopyTextureInfo.Size = FIntVector(DestinationTexture->Desc.GetSize().X, DestinationTexture->Desc.GetSize().Y, 0);
+		AddCopyTexturePass(GraphBuilder, DestinationTexture, InputCopy, CopyTextureInfo);
+
+		FLandscapeCircleHeightPatchPS::FParameters* ShaderParams =
+			GraphBuilder.AllocParameters<FLandscapeCircleHeightPatchPS::FParameters>();
+		*ShaderParams = ShaderParamsToCopy;
+
+		FRDGTextureSRVRef InputCopySRV = GraphBuilder.CreateSRV(FRDGTextureSRVDesc::CreateForMipLevel(InputCopy, 0));
+		ShaderParams->InSourceHeightmap = InputCopySRV;
+
+		ShaderParams->RenderTargets[0] = FRenderTargetBinding(DestinationTexture, ERenderTargetLoadAction::ENoAction, /*InMipIndex = */0);
+
+		FLandscapeCircleHeightPatchPS::AddToRenderGraph(GraphBuilder, ShaderParams);
+
+		GraphBuilder.Execute();
+	});
+
+	return InCombinedResult;
+}
+
+#undef LOCTEXT_NAMESPACE
\ No newline at end of file
diff --git a/Engine/Plugins/Experimental/LandscapePatch/Source/LandscapePatch/Private/LandscapeCircleHeightPatchPS.cpp b/Engine/Plugins/Experimental/LandscapePatch/Source/LandscapePatch/Private/LandscapeCircleHeightPatchPS.cpp
new file mode 100644
--- /dev/null
+++ b/Engine/Plugins/Experimental/LandscapePatch/Source/LandscapePatch/Private/LandscapeCircleHeightPatchPS.cpp
@@ -0,0 +1,37 @@
+// Copyright Epic Games, Inc. All Rights Reserved.
+
+#include "LandscapeCircleHeightPatchPS.h"
+
+#include "PixelShaderUtils.h"
+
+bool FLandscapeCircleHeightPatchPS::ShouldCompilePermutation(const FGlobalShaderPermutationParameters& Parameters)
+{
+	// Apparently landscape requires a particular feature level
+	return IsFeatureLevelSupported(Parameters.Platform, ERHIFeatureLevel::SM5)
+		&& !IsConsolePlatform(Parameters.Platform)
+		&& !IsMetalMobilePlatform(Parameters.Platform);
+}
+
+void FLandscapeCircleHeightPatchPS::ModifyCompilationEnvironment(const FGlobalShaderPermutationParameters& Parameters, FShaderCompilerEnvironment& OutEnvironment)
+{
+	// If we end up using an ifdef
+	//OutEnvironment.SetDefine(TEXT("CIRCLE_HEIGHT_PATCH"), 1);
+}
+
+void FLandscapeCircleHeightPatchPS::AddToRenderGraph(FRDGBuilder& GraphBuilder, FParameters* InParameters)
+{
+	FGlobalShaderMap* ShaderMap = GetGlobalShaderMap(GMaxRHIFeatureLevel);
+	TShaderMapRef<FLandscapeCircleHeightPatchPS> PixelShader(ShaderMap);
+
+	FIntVector TextureSize = InParameters->InSourceHeightmap->Desc.Texture->Desc.GetSize();
+
+	FPixelShaderUtils::AddFullscreenPass(
+		GraphBuilder,
+		ShaderMap,
+		RDG_EVENT_NAME("LandmassCircleHeightPatch"),
+		PixelShader,
+		InParameters,
+		FIntRect(0, 0, TextureSize.X, TextureSize.Y));
+}
+
+IMPLEMENT_GLOBAL_SHADER(FLandscapeCircleHeightPatchPS, "/Plugin/LandscapePatch/Private/LandscapeCircleHeightPatchPS.usf", "ApplyLandscapeCircleHeightPatch", SF_Pixel);
\ No newline at end of file
diff --git a/Engine/Plugins/Experimental/LandscapePatch/Source/LandscapePatch/Private/LandscapePatch.cpp b/Engine/Plugins/Experimental/LandscapePatch/Source/LandscapePatch/Private/LandscapePatch.cpp
new file mode 100644
--- /dev/null
+++ b/Engine/Plugins/Experimental/LandscapePatch/Source/LandscapePatch/Private/LandscapePatch.cpp
@@ -0,0 +1,235 @@
+// Copyright Epic Games, Inc. All Rights Reserved.
+
+#include "LandscapePatch.h"
+
+#include "EngineUtils.h"
+#include "GameFramework/Actor.h"
+#include "Landscape.h"
+#include "LandscapePatchLogging.h"
+#include "LandscapePatchManager.h"
+
+#define LOCTEXT_NAMESPACE "LandscapePatch"
+
+#if WITH_EDITOR
+namespace LandscapePatchLocals
+{
+	ALandscapePatchManager* CreateNewPatchManagerForLandscape(ALandscape* Landscape)
+	{
+		if (!ensure(Landscape->CanHaveLayersContent()))
+		{
+			return nullptr;
+		}
+
+		FString BrushActorString = FString::Format(TEXT("{0}_{1}"), { Landscape->GetActorLabel(), ALandscapePatchManager::StaticClass()->GetName() });
+		FName BrushActorName = MakeUniqueObjectName(Landscape->GetOuter(), ALandscapePatchManager::StaticClass(), FName(BrushActorString));
+		FActorSpawnParameters SpawnParams;
+		SpawnParams.Name = BrushActorName;
+		SpawnParams.bAllowDuringConstructionScript = true; // This can be called by construction script if the actor being added to the world is part of a blueprint
+
+		ALandscapePatchManager* PatchManager = Landscape->GetWorld()->SpawnActor<ALandscapePatchManager>(ALandscapePatchManager::StaticClass(), SpawnParams);
+		if (PatchManager)
+		{
+			PatchManager->SetActorLabel(BrushActorString);
+			PatchManager->SetTargetLandscape(Landscape);
+		}
+
+		return PatchManager;
+	}
+}
+
+ULandscapePatchComponent::ULandscapePatchComponent(const FObjectInitializer& ObjectInitializer)
+	: Super(ObjectInitializer)
+{
+	// Causes OnUpdateTransform to be called when the parent is moved. Note that this is better to do here in the
+	// constructor, otherwise we'd need to do it both in OnComponentCreated and PostLoad.
+	// We could keep this false if we were to register to TransformUpdated, since that gets broadcast either way.
+	// TODO: Currently, neither TransformUpdated nor OnUpdateTransform are triggered when parent's transform is changed
+	bWantsOnUpdateTransform = true;
+}
+
+void ULandscapePatchComponent::OnComponentCreated()
+{
+	using namespace LandscapePatchLocals;
+
+	Super::OnComponentCreated();
+		
+	// Mark whether we're creating from scratch of from a copy
+	bWasCopy = bPropertiesCopiedIndicator;
+	bPropertiesCopiedIndicator = true;
+
+	UWorld* World = GetWorld();
+	AActor* OwningActor = GetOwner();
+
+	if (PatchManager.IsValid())
+	{
+		// If we copied over a patch manager, presumably Landscape should be
+		// copied over as well, but might as well do this to be safe.
+		Landscape = PatchManager->GetOwningLandscape();
+	}
+	else if (World)
+	{
+		if (Landscape.IsValid())
+		{
+			// If we copied over a patch with a landscape but no manager, create manager in that landscape
+			if (Landscape->CanHaveLayersContent())
+			{
+				PatchManager = CreateNewPatchManagerForLandscape(Landscape.Get());
+			}
+			else
+			{
+				UE_LOG(LogLandscapePatch, Warning, TEXT("Landscape target for height patch did not have edit layers enabled. Unable to create patch manager."));
+				Landscape = nullptr;
+			}
+		}
+		else
+		{
+			// Didn't copy over an existing manager or landscape.
+
+			// See if the level has a height patch manager to which we can add ourselves
+			TActorIterator<ALandscapePatchManager> ManagerIterator(World);
+			if (!!ManagerIterator)
+			{
+				PatchManager = *ManagerIterator;
+				Landscape = PatchManager->GetOwningLandscape();
+			}
+			else
+			{
+				// If no existing manager, find some landscape and add a new one.
+				for (TActorIterator<ALandscape> LandscapeIterator(World); LandscapeIterator; ++LandscapeIterator)
+				{
+					if (LandscapeIterator->CanHaveLayersContent())
+					{
+						Landscape = *LandscapeIterator;
+						PatchManager = CreateNewPatchManagerForLandscape(Landscape.Get());
+						break;
+					}
+				}
+				if (!PatchManager.IsValid())
+				{
+					UE_LOG(LogLandscapePatch, Warning, TEXT("Unable to find a landscape with edit layers enabled. Unable to create patch manager."));
+				}
+			}
+		}
+	}
+
+	if (PatchManager.IsValid())
+	{
+		PatchManager->AddPatch(this);
+	}
+}
+
+void ULandscapePatchComponent::PostLoad()
+{
+	Super::PostLoad();
+
+	bPropertiesCopiedIndicator = true;
+	bLoadedButNotYetRegistered = true;
+}
+
+void ULandscapePatchComponent::OnComponentDestroyed(bool bDestroyingHierarchy)
+{
+	if (PatchManager.IsValid())
+	{
+		PatchManager->RemovePatch(this);
+	}
+
+	Super::OnComponentDestroyed(bDestroyingHierarchy);
+}
+
+void ULandscapePatchComponent::OnRegister()
+{
+	Super::OnRegister();
+
+
+	if (bLoadedButNotYetRegistered)
+	{
+		bLoadedButNotYetRegistered = false;
+		return;
+	}
+
+	// TODO: Currently the main reason to invalidate the landscape on registration is to respond
+	// to detail panel changes of the parent's transform. However we may be able to catch a wide
+	// variety of changes this way, so we'll need to see if we can get rid of other invalidations.
+	// Also, we should make the invalidation conditional on whether we actually modify any relevant
+	// properties by having a virtual method that compares and updates a stored hash of them.
+	if (PatchManager.IsValid())
+	{
+		PatchManager->RequestLandscapeUpdate();
+	}
+}
+
+void ULandscapePatchComponent::OnUpdateTransform(EUpdateTransformFlags UpdateTransformFlags, ETeleportType Teleport)
+{
+	Super::OnUpdateTransform(UpdateTransformFlags, Teleport);
+
+	if (PatchManager.IsValid())
+	{
+		PatchManager->RequestLandscapeUpdate();
+	}
+}
+
+void ULandscapePatchComponent::PostEditChangeProperty(FPropertyChangedEvent& PropertyChangedEvent)
+{
+	using namespace LandscapePatchLocals;
+
+	Super::PostEditChangeProperty(PropertyChangedEvent);
+
+	// If we're changing the owning landscape, there's some work we need to do to remove/add ourselves from/to the proper
+	// brush managers.
+	if (PropertyChangedEvent.Property 
+		&& (PropertyChangedEvent.Property->GetFName() == GET_MEMBER_NAME_CHECKED(ULandscapePatchComponent, Landscape)))
+	{
+		if (PatchManager.IsValid())
+		{
+			PatchManager->RemovePatch(this);
+		}
+		PatchManager = nullptr;
+
+		bool bAdded = !Landscape.IsValid();
+		for (TActorIterator<ALandscapePatchManager> ManagersItr(GetWorld()); ManagersItr; ++ManagersItr)
+		{
+			// Exit early if found the relevant manager.
+			if (bAdded)
+			{
+				break;
+			}
+
+			ALandscapePatchManager* Manager = *ManagersItr;
+			if (!ensure(Manager->IsValidLowLevel()))
+			{
+				continue;
+			}
+
+			if (!bAdded && Manager->GetOwningLandscape() == Landscape)
+			{
+				Manager->AddPatch(this);
+				PatchManager = Manager;
+				bAdded = true;
+			}
+		}
+
+		// If we had a non-null landscape and didn't add ourselves to its manager despite iterating through
+		// all managers, then the landscape must not have had a manager. Create one for it.
+		if (!bAdded)
+		{
+			if (Landscape->CanHaveLayersContent())
+			{
+				PatchManager = CreateNewPatchManagerForLandscape(Landscape.Get());
+				PatchManager->AddPatch(this);
+			}
+			else
+			{
+				UE_LOG(LogLandscapePatch, Warning, TEXT("Landscape target for height patch did not have edit layers enabled. Unable to create patch manager."));
+				Landscape = nullptr;
+			}
+		}
+	}
+
+	if (PatchManager.IsValid())
+	{
+		PatchManager->RequestLandscapeUpdate();
+	}
+}
+#endif
+
+#undef LOCTEXT_NAMESPACE
diff --git a/Engine/Plugins/Experimental/LandscapePatch/Source/LandscapePatch/Private/LandscapePatchLogging.cpp b/Engine/Plugins/Experimental/LandscapePatch/Source/LandscapePatch/Private/LandscapePatchLogging.cpp
new file mode 100644
--- /dev/null
+++ b/Engine/Plugins/Experimental/LandscapePatch/Source/LandscapePatch/Private/LandscapePatchLogging.cpp
@@ -0,0 +1,5 @@
+// Copyright Epic Games, Inc. All Rights Reserved.
+
+#include "LandscapePatchLogging.h"
+
+DEFINE_LOG_CATEGORY(LogLandscapePatch);
diff --git a/Engine/Plugins/Experimental/LandscapePatch/Source/LandscapePatch/Private/LandscapePatchManager.cpp b/Engine/Plugins/Experimental/LandscapePatch/Source/LandscapePatch/Private/LandscapePatchManager.cpp
new file mode 100644
--- /dev/null
+++ b/Engine/Plugins/Experimental/LandscapePatch/Source/LandscapePatch/Private/LandscapePatchManager.cpp
@@ -0,0 +1,160 @@
+// Copyright Epic Games, Inc. All Rights Reserved.
+
+#include "LandscapePatchManager.h"
+
+#include "Landscape.h"
+#include "LandscapePatch.h"
+#include "Engine/TextureRenderTarget2D.h"
+#include "LandscapeDataAccess.h"
+#include "RenderGraph.h" // RDG_EVENT_NAME
+
+#define LOCTEXT_NAMESPACE "LandscapePatchManager"
+
+// TODO: Not sure if using this kind of constructor is a proper thing to do vs some other hook...
+ALandscapePatchManager::ALandscapePatchManager(const FObjectInitializer& ObjectInitializer)
+	: ALandscapeBlueprintBrushBase(ObjectInitializer)
+{
+#if WITH_EDITOR
+	SetAffectsHeightmap(true);
+	SetAffectsWeightmap(true);
+#endif
+}
+
+void ALandscapePatchManager::Initialize_Native(const FTransform & InLandscapeTransform,
+	const FIntPoint& InLandscapeSize,
+	const FIntPoint& InLandscapeRenderTargetSize)
+{
+	// Get a transform from pixel coordinate in heightmap to world space coordinate. Note that we can't
+	// store the inverse directly because a FTransform can't properly represent a TRS inverse when the
+	// original TRS has non-uniform scaling).
+
+	// The pixel to landscape-space transform is unrotated, (S_p * x + T_p). The landscape to world
+	// transform gets applied on top of this: (R_l * S_l * (S_p * x + T_p)) + T_L. Collapsing this
+	// down to pixel to world TRS, we get: R_l * (S_l * S_p) * x + (R_l * S_l * T_p + T_L)
+
+	// To go from stored height value to unscaled height, we divide by 128 and subtract 256. We can get these
+	// values from the constants in LandscapeDataAccess.h (we distribute the multiplication by LANDSCAPE_ZSCALE
+	// so that translation happens after scaling like in TRS)
+	const double HEIGHTMAP_TO_OBJECT_HEIGHT_SCALE = LANDSCAPE_ZSCALE;
+	const double HEIGHTMAP_TO_OBJECT_HEIGHT_OFFSET = -LandscapeDataAccess::MidValue * LANDSCAPE_ZSCALE;
+
+	// S_p: the pixel coordinates are actually the same as xy object-space coordinates becaues one quad is 1 unit,
+	// so we only need to scale the height.
+	FVector3d PixelToObjectSpaceScale = FVector3d(
+		1,
+		1,
+		HEIGHTMAP_TO_OBJECT_HEIGHT_SCALE
+	);
+
+	// T_p: again, only the height needs translation
+	FVector3d PixelToObjectSpaceTranslate = FVector3d(
+		0,
+		0,
+		HEIGHTMAP_TO_OBJECT_HEIGHT_OFFSET
+	);
+
+	// S_l* S_p: composed scale
+	HeightmapCoordsToWorld.SetScale3D(InLandscapeTransform.GetScale3D() * PixelToObjectSpaceScale);
+
+	// R_l
+	HeightmapCoordsToWorld.SetRotation(InLandscapeTransform.GetRotation());
+
+	// R_l * S_l * T_p + T_L: composed translation
+	HeightmapCoordsToWorld.SetTranslation(InLandscapeTransform.TransformVector(PixelToObjectSpaceTranslate)
+		+ InLandscapeTransform.GetTranslation());
+}
+
+UTextureRenderTarget2D* ALandscapePatchManager::Render_Native(bool InIsHeightmap,
+	UTextureRenderTarget2D* InCombinedResult,
+	const FName& InWeightmapLayerName)
+{
+	// Used to determine whether we need to remove any invalid brushes
+	bool bHaveInvalidPatches = false;
+
+	// TODO: There are many uncertainties in how we iterate across the height patches and have them
+	// apply themselves. For one thing we may want to pass around a render graph, in which case this
+	// loop will happen on the render thread somehow. For another, it's not yet determined what all
+	// of this will look like when we have the ability to render to just a subsection of the entire
+	// height map.
+	// So for now we do the simplest thing, and that is to have the height patches act as if they were
+	// independent brushes.
+	for (TWeakObjectPtr<ULandscapePatchComponent>& Component : PatchComponents)
+	{
+		// Theoretically when components are marked for destruction, they should remove themselves from
+		// the patch manager, hence the ensure here
+		if (ensure(Component.IsValid()))
+		{
+			InCombinedResult = Component->Render_Native(InIsHeightmap, InCombinedResult, InWeightmapLayerName);
+		}
+		else
+		{
+			bHaveInvalidPatches = true;
+		}
+	}
+
+	if (bHaveInvalidPatches)
+	{
+		PatchComponents.RemoveAll([](TWeakObjectPtr<ULandscapePatchComponent> Component) {
+			return !Component.IsValid();
+		});
+	}
+
+	return InCombinedResult;
+}
+
+void ALandscapePatchManager::SetTargetLandscape(ALandscape* InTargetLandscape)
+{
+#if WITH_EDITOR
+	if (OwningLandscape != InTargetLandscape)
+	{
+		if (OwningLandscape)
+		{
+			OwningLandscape->RemoveBrush(this);
+		}
+
+		if (InTargetLandscape && ensure(InTargetLandscape->CanHaveLayersContent()))
+		{
+			FName LayerName = FName("LandscapePatches");
+			int32 ExistingPatchLayerIndex = InTargetLandscape->GetLayerIndex(LayerName);
+			if (ExistingPatchLayerIndex == INDEX_NONE)
+			{
+				ExistingPatchLayerIndex = InTargetLandscape->CreateLayer(LayerName);
+			}
+			InTargetLandscape->AddBrushToLayer(ExistingPatchLayerIndex, this);
+		}
+	}
+#endif
+}
+
+void ALandscapePatchManager::AddPatch(TObjectPtr<ULandscapePatchComponent> Patch)
+{
+	if (Patch)
+	{
+		Modify();
+		PatchComponents.AddUnique(Patch);
+		RequestLandscapeUpdate();
+	}
+}
+
+bool ALandscapePatchManager::RemovePatch(TObjectPtr<ULandscapePatchComponent> Patch)
+{
+	bool bRemoved = false;
+
+	if (Patch)
+	{
+		Modify();
+		bRemoved = PatchComponents.Remove(Patch) > 0;
+		RequestLandscapeUpdate();
+	}
+	
+	return bRemoved;
+}
+
+#if WITH_EDITOR
+void ALandscapePatchManager::PostEditUndo()
+{
+	RequestLandscapeUpdate();
+}
+#endif
+
+#undef LOCTEXT_NAMESPACE
\ No newline at end of file
diff --git a/Engine/Plugins/Experimental/LandscapePatch/Source/LandscapePatch/Private/LandscapePatchModule.cpp b/Engine/Plugins/Experimental/LandscapePatch/Source/LandscapePatch/Private/LandscapePatchModule.cpp
new file mode 100644
--- /dev/null
+++ b/Engine/Plugins/Experimental/LandscapePatch/Source/LandscapePatch/Private/LandscapePatchModule.cpp
@@ -0,0 +1,23 @@
+// Copyright Epic Games, Inc. All Rights Reserved.
+
+#include "LandscapePatchModule.h"
+
+#include "Interfaces/IPluginManager.h"
+#include "ShaderCore.h" // AddShaderSourceDirectoryMapping
+
+#define LOCTEXT_NAMESPACE "FLandscapePatchModule"
+
+void FLandscapePatchModule::StartupModule()
+{
+	// Maps virtual shader source directory to the plugin's actual Shaders directory.
+	FString PluginShaderDir = FPaths::Combine(IPluginManager::Get().FindPlugin(TEXT("LandscapePatch"))->GetBaseDir(), TEXT("Shaders"));
+	AddShaderSourceDirectoryMapping(TEXT("/Plugin/LandscapePatch"), PluginShaderDir);
+}
+
+void FLandscapePatchModule::ShutdownModule()
+{
+}
+
+#undef LOCTEXT_NAMESPACE
+	
+IMPLEMENT_MODULE(FLandscapePatchModule, LandscapePatch)
\ No newline at end of file
diff --git a/Engine/Plugins/Experimental/LandscapePatch/Source/LandscapePatch/Public/LandscapeCircleHeightPatch.h b/Engine/Plugins/Experimental/LandscapePatch/Source/LandscapePatch/Public/LandscapeCircleHeightPatch.h
new file mode 100644
--- /dev/null
+++ b/Engine/Plugins/Experimental/LandscapePatch/Source/LandscapePatch/Public/LandscapeCircleHeightPatch.h
@@ -0,0 +1,48 @@
+// Copyright Epic Games, Inc. All Rights Reserved.
+
+#pragma once
+
+#include "CoreMinimal.h"
+
+#include "LandscapePatch.h"
+
+#include "LandscapeCircleHeightPatch.generated.h"
+
+/**
+ * The simplest height patch: a circle of flat ground with a falloff past the initial radius across which the
+ * alpha decreases linearly. When added to an actor, initializes itself to the bottom of the bounding box.
+ */
+UCLASS(Blueprintable, BlueprintType, ClassGroup = Landscape, meta = (BlueprintSpawnableComponent))
+class LANDSCAPEPATCH_API ULandscapeCircleHeightPatch : public ULandscapePatchComponent
+{
+	GENERATED_BODY()
+
+public:
+
+	//TODO: For now the height patches have a similar interface to blueprint brushes, but this is likely to
+	// change. We are likely to pass a graph builder to the patches instead, along with an area that
+	// we want affected.
+	virtual void Initialize_Native(const FTransform& InLandscapeTransform, 
+		const FIntPoint& InLandscapeSize, 
+		const FIntPoint& InLandscapeRenderTargetSize) override;
+	virtual UTextureRenderTarget2D* Render_Native(bool InIsHeightmap,
+		UTextureRenderTarget2D* InCombinedResult,
+		const FName& InWeightmapLayerName) override;
+
+	// UActorComponent
+	virtual void OnComponentCreated() override;
+
+protected:
+
+	UPROPERTY(EditAnywhere, Category = Settings)
+	float Radius = 500;
+
+	/** Distance across which the alpha will go from 1 down to 0 outside of circle. */
+	UPROPERTY(EditAnywhere, Category = Settings)
+	float Falloff = 500;
+
+	/** When true, only the vertices in the circle have alpha 1. If false, the radius is expanded slightly so that neighboring 
+	  vertices are also included and the whole circle is able to lie flat. */
+	UPROPERTY(EditAnywhere, Category = Settings, AdvancedDisplay)
+	bool bExclusiveRadius = false;
+};
\ No newline at end of file
diff --git a/Engine/Plugins/Experimental/LandscapePatch/Source/LandscapePatch/Public/LandscapeCircleHeightPatchPS.h b/Engine/Plugins/Experimental/LandscapePatch/Source/LandscapePatch/Public/LandscapeCircleHeightPatchPS.h
new file mode 100644
--- /dev/null
+++ b/Engine/Plugins/Experimental/LandscapePatch/Source/LandscapePatch/Public/LandscapeCircleHeightPatchPS.h
@@ -0,0 +1,32 @@
+// Copyright Epic Games, Inc. All Rights Reserved.
+
+#pragma once
+
+#include "CoreMinimal.h"
+
+#include "GlobalShader.h"
+#include "ShaderParameterStruct.h"
+
+/**
+ * Shader that applies a circle height patch to a landscape heightmap.
+ */
+class LANDSCAPEPATCH_API FLandscapeCircleHeightPatchPS : public FGlobalShader
+{
+	DECLARE_GLOBAL_SHADER(FLandscapeCircleHeightPatchPS);
+	SHADER_USE_PARAMETER_STRUCT(FLandscapeCircleHeightPatchPS, FGlobalShader);
+
+public:
+	BEGIN_SHADER_PARAMETER_STRUCT(FParameters, )
+		SHADER_PARAMETER_RDG_TEXTURE_SRV(Texture2D<float4>, InSourceHeightmap) // Our input texture
+		SHADER_PARAMETER(FVector3f, InCenter)
+		SHADER_PARAMETER(float, InRadius)
+		SHADER_PARAMETER(float, InFalloff)
+		RENDER_TARGET_BINDING_SLOTS() // Holds our output
+		END_SHADER_PARAMETER_STRUCT()
+
+	static bool ShouldCompilePermutation(const FGlobalShaderPermutationParameters& Parameters);
+
+	static void ModifyCompilationEnvironment(const FGlobalShaderPermutationParameters& Parameters, FShaderCompilerEnvironment& OutEnvironment);
+
+	static void AddToRenderGraph(FRDGBuilder& GraphBuilder, FParameters* InParameters);
+};
\ No newline at end of file
diff --git a/Engine/Plugins/Experimental/LandscapePatch/Source/LandscapePatch/Public/LandscapePatch.h b/Engine/Plugins/Experimental/LandscapePatch/Source/LandscapePatch/Public/LandscapePatch.h
new file mode 100644
--- /dev/null
+++ b/Engine/Plugins/Experimental/LandscapePatch/Source/LandscapePatch/Public/LandscapePatch.h
@@ -0,0 +1,74 @@
+// Copyright Epic Games, Inc. All Rights Reserved.
+
+#pragma once
+
+#include "CoreMinimal.h"
+
+#include "Components/SceneComponent.h"
+
+#include "LandscapePatch.generated.h"
+
+class ALandscape;
+class ALandscapePatchManager;
+class UTextureRenderTarget2D;
+
+/**
+ * Base class for landscape patches: components that can be attached to meshes and moved around to make
+ * the meshes affect the landscape around themselves.
+ */
+//~ TODO: Although this doesn't generate geometry, we are likely to change this to inherit from UPrimitiveComponent
+//~ so that we can use render proxies for passing along data to the render thread or perhaps for visualization.
+UCLASS(Blueprintable, BlueprintType, Abstract)
+class LANDSCAPEPATCH_API ULandscapePatchComponent : public USceneComponent
+{
+	GENERATED_BODY()
+
+public:
+	ULandscapePatchComponent(const FObjectInitializer& ObjectInitializer = FObjectInitializer::Get());
+
+	virtual void Initialize_Native(const FTransform& InLandscapeTransform,
+		const FIntPoint& InLandscapeSize,
+		const FIntPoint& InLandscapeRenderTargetSize) {}
+	virtual UTextureRenderTarget2D* Render_Native(bool InIsHeightmap,
+		UTextureRenderTarget2D* InCombinedResult,
+		const FName& InWeightmapLayerName) { return InCombinedResult; }
+
+// For now we keep the patches largely editor-only, since we don't yet support runtime landscape editing
+#if WITH_EDITOR
+	// USceneComponent
+	virtual void OnUpdateTransform(EUpdateTransformFlags UpdateTransformFlags, ETeleportType Teleport) override;
+
+	// UActorComponent
+	virtual void OnComponentCreated() override;
+	virtual void OnComponentDestroyed(bool bDestroyingHierarchy) override;
+	virtual void OnRegister() override;
+
+	// UObject
+	virtual void PostEditChangeProperty(FPropertyChangedEvent& PropertyChangedEvent) override;
+	virtual bool IsPostLoadThreadSafe() const override { return true; }
+	virtual void PostLoad();
+#endif
+	virtual bool IsEditorOnly() const override { return true; }
+	virtual bool NeedsLoadForClient() const override { return false; }
+	virtual bool NeedsLoadForServer() const override { return false; }
+
+protected:
+	UPROPERTY(EditAnywhere, Category = Settings)
+	TWeakObjectPtr<ALandscape> Landscape = nullptr;
+
+	UPROPERTY()
+	TWeakObjectPtr<ALandscapePatchManager> PatchManager = nullptr;
+
+	// Determines whether the height patch was made by copying a different height patch.
+	bool bWasCopy = false;
+
+	// This is true for existing height patches right after they are loaded, so that we can ignore
+	// the first OnRegister call. It remains false from the first OnRegiter call onward, even
+	// if the component is unregistered.
+	bool bLoadedButNotYetRegistered = false;
+private:
+	// Starts as false and gets set to true in construction, so gets used to set bWasCopy
+	// by checking the indicator value at the start of construction.
+	UPROPERTY()
+	bool bPropertiesCopiedIndicator = false;
+};
\ No newline at end of file
diff --git a/Engine/Plugins/Experimental/LandscapePatch/Source/LandscapePatch/Public/LandscapePatchLogging.h b/Engine/Plugins/Experimental/LandscapePatch/Source/LandscapePatch/Public/LandscapePatchLogging.h
new file mode 100644
--- /dev/null
+++ b/Engine/Plugins/Experimental/LandscapePatch/Source/LandscapePatch/Public/LandscapePatchLogging.h
@@ -0,0 +1,7 @@
+// Copyright Epic Games, Inc. All Rights Reserved.
+
+#pragma once
+
+#include "CoreMinimal.h"
+
+LANDSCAPEPATCH_API DECLARE_LOG_CATEGORY_EXTERN(LogLandscapePatch, Log, All);
diff --git a/Engine/Plugins/Experimental/LandscapePatch/Source/LandscapePatch/Public/LandscapePatchManager.h b/Engine/Plugins/Experimental/LandscapePatch/Source/LandscapePatch/Public/LandscapePatchManager.h
new file mode 100644
--- /dev/null
+++ b/Engine/Plugins/Experimental/LandscapePatch/Source/LandscapePatch/Public/LandscapePatchManager.h
@@ -0,0 +1,68 @@
+// Copyright Epic Games, Inc. All Rights Reserved.
+
+#pragma once
+
+#include "CoreMinimal.h"
+
+#include "LandscapeBlueprintBrushBase.h"
+
+#include "LandscapePatchManager.generated.h"
+
+class ULandscapePatchComponent;
+
+/**
+ * Acts as the "blueprint brush" as far as the owning edit layer is concerned. In reality, has its contained
+ * patches edit the height/weight maps.
+ */
+//~ The alternative to this approach is to have the individual patches act as independent brushes, which
+//~ we currently don't want to do because we think it will clutter the brush interface and may lose opportunities
+//~ for optimization... 
+UCLASS()
+class LANDSCAPEPATCH_API ALandscapePatchManager : public ALandscapeBlueprintBrushBase
+{
+	GENERATED_BODY()
+
+public:
+
+	ALandscapePatchManager(const FObjectInitializer& ObjectInitializer = FObjectInitializer::Get());
+
+	// These get called by the landscape system to apply the patches to the height/weight maps.
+	virtual void Initialize_Native(const FTransform& InLandscapeTransform,
+		const FIntPoint& InLandscapeSize,
+		const FIntPoint& InLandscapeRenderTargetSize) override;
+	virtual UTextureRenderTarget2D* Render_Native(bool InIsHeightmap,
+		UTextureRenderTarget2D* InCombinedResult,
+		const FName& InWeightmapLayerName) override;
+
+	// Adds the brush to the given landscape, removing it from any previous one. This differs from SetOwningLandscape
+	// in that SetOwningLandscape is called by the landscape itself from AddBrushToLayer to update the manager.
+	virtual void SetTargetLandscape(ALandscape* InOwningLandscape);
+
+	// For use by the owned patch objects.
+	/**
+	 * Gets the transform from a point in the heightmap (where x and y are pixel coordinates,
+	 * aka coordinates of the associated vertex, and z is the height as stored in the height
+	 * map, currently a 16 bit integer) to world point based on the current landscape transform.
+	 */
+	virtual FTransform GetHeightmapCoordsToWorld() { return HeightmapCoordsToWorld; }
+
+	void AddPatch(TObjectPtr<ULandscapePatchComponent> Patch);
+
+	bool RemovePatch(TObjectPtr<ULandscapePatchComponent> Patch);
+
+	// UObject
+#if WITH_EDITOR
+	virtual void PostEditUndo() override;
+#endif
+	virtual bool IsEditorOnly() const override { return true; }
+	virtual bool NeedsLoadForClient() const override { return false; }
+	virtual bool NeedsLoadForServer() const override { return false; }
+
+protected:
+	
+	UPROPERTY()
+	TArray<TWeakObjectPtr<ULandscapePatchComponent>> PatchComponents;
+
+	FTransform HeightmapCoordsToWorld;
+
+};
\ No newline at end of file
diff --git a/Engine/Plugins/Experimental/LandscapePatch/Source/LandscapePatch/Public/LandscapePatchModule.h b/Engine/Plugins/Experimental/LandscapePatch/Source/LandscapePatch/Public/LandscapePatchModule.h
new file mode 100644
--- /dev/null
+++ b/Engine/Plugins/Experimental/LandscapePatch/Source/LandscapePatch/Public/LandscapePatchModule.h
@@ -0,0 +1,15 @@
+// Copyright Epic Games, Inc. All Rights Reserved.
+
+#pragma once
+
+#include "CoreMinimal.h"
+#include "Modules/ModuleManager.h"
+
+class FLandscapePatchModule : public IModuleInterface
+{
+public:
+
+	/** IModuleInterface implementation */
+	virtual void StartupModule() override;
+	virtual void ShutdownModule() override;
+};
