Landscape: Adding landscape texture height patches, with support for internal and external texture source modes, additive and alpha blend modes, custom alpha, circle and rectangle falloffs, blueprint support, and initialization from existing landscape.

diff --git a/Engine/Plugins/Experimental/LandscapePatch/LandscapePatch.uplugin b/Engine/Plugins/Experimental/LandscapePatch/LandscapePatch.uplugin
--- a/Engine/Plugins/Experimental/LandscapePatch/LandscapePatch.uplugin
+++ b/Engine/Plugins/Experimental/LandscapePatch/LandscapePatch.uplugin
@@ -19,6 +19,11 @@
 			"Name": "LandscapePatch",
 			"Type": "Runtime",
 			"LoadingPhase": "PostConfigInit"
+		},
+		{
+			"Name": "LandscapePatchEditorOnly",
+			"Type": "Editor",
+			"LoadingPhase": "PostEngineInit"
 		}
 	]
 }
\ No newline at end of file
diff --git a/Engine/Plugins/Experimental/LandscapePatch/Shaders/Private/LandscapeTextureHeightPatchPS.usf b/Engine/Plugins/Experimental/LandscapePatch/Shaders/Private/LandscapeTextureHeightPatchPS.usf
new file mode 100644
--- /dev/null
+++ b/Engine/Plugins/Experimental/LandscapePatch/Shaders/Private/LandscapeTextureHeightPatchPS.usf
@@ -0,0 +1,148 @@
+// Copyright Epic Games, Inc. All Rights Reserved.
+
+#pragma once
+
+#include "/Engine/Private/Common.ush"
+#include "/Engine/Private/Landscape/LandscapeCommon.ush"
+
+#if defined (__INTELLISENSE__)
+// Uncomment the appropriate define for enabling syntax highlighting with HLSL Tools for Visual Studio : 
+//#define APPLY_LANDSCAPE_PATCH 1
+//#define OFFSET_LANDSCAPE_PATCH 1
+//#define SIMPLE_TEXTURE_COPY 1
+#endif // defined (__INTELLISENSE__)
+
+#if APPLY_LANDSCAPE_PATCH
+Texture2D<float4> InSourceHeightmap;
+Texture2D<float4> InHeightPatch;
+SamplerState InHeightPatchSampler;
+float4x4 InHeightmapToPatch;
+float2 InPatchWorldDimensions;
+float InPatchBaseHeight;
+float2 InEdgeUVDeadBorder;
+float InFalloffWorldMargin;
+float InHeightScale;
+// When false, falloff is circular. When true, it is rectangular.
+uint bRectangularFalloff;
+// When true, the texture alpha channel is considered for blending (in addition to falloff, if any)
+uint bApplyPatchAlpha;
+// When true, the patch stored height is interpreted relative to landscape mid value and applied to
+// the landscape, instead of alpha blending with the landscape.
+uint bAdditiveMode;
+
+void ApplyLandscapeTextureHeightPatch(in float4 SVPos : SV_POSITION, out float2 OutColor : SV_Target0)
+{
+	static const float HALF_PI = 3.14159265f / 2;
+	static const float LANDSCAPE_MID_VALUE = 32768.0f;
+	static const float FLOAT_TOLERANCE = 0.0001;
+	
+	// We need only the 2D affine transformation that goes from landscape XY heightmap integer coordinates
+	// to patch UV coordinates.
+	float2x2 HeightmapToPatchRotateScale = (float2x2) InHeightmapToPatch;
+	float2 HeightmapToPatchTranslate = InHeightmapToPatch._m03_m13;
+	
+	float2 PatchUVCoordinates = mul(HeightmapToPatchRotateScale, SVPos.xy) + HeightmapToPatchTranslate;
+	float4 PatchSampledValue = InHeightPatch.Sample(InHeightPatchSampler, PatchUVCoordinates);
+	float PatchStoredHeight = UnpackHeight(PatchSampledValue.xy);
+	if (bAdditiveMode)
+	{
+		PatchStoredHeight -= LANDSCAPE_MID_VALUE;
+	}
+	float PatchHeight = InHeightScale * PatchStoredHeight + InPatchBaseHeight;
+	
+	int2 HeightmapCoordinates = floor(SVPos.xy);
+	float4 CurrentPackedHeight = InSourceHeightmap.Load(int3(HeightmapCoordinates, 0));
+	float CurrentHeight = UnpackHeight(CurrentPackedHeight.xy);
+	
+	// The falloff starts at the edges and goes inward, and shouldn't be higher than 1/2 of the smaller dimension, or
+	// else it extends past the center. If it does extend past the center, we'll decrease the max alpha in the center
+	// (imagine shaving off a cube with progressively flatter slopes- once you get past the center of the top, you get
+	// a progressively shorter pyramid). Maybe we shouldn't even bother dealing with this case, but it's easy enough to do.
+	float ClampedFalloffMargin = clamp(InFalloffWorldMargin, 0, min(InPatchWorldDimensions.x, InPatchWorldDimensions.y) / 2);
+	float bFalloffAlpha = InFalloffWorldMargin > 0 ? ClampedFalloffMargin / InFalloffWorldMargin : 1;
+
+	// Alpha T is a parameter that we'll pass through a cos^2 function to get the alpha (to have it be smoother). It should be 
+	// 0 at the places where there is no falloff (ie alpha is 1), and go to 1 at places where falloff is complete (and falloff is 0)
+	float AlphaT = 1;
+	
+	if (bRectangularFalloff)
+	{
+		float2 DistancesFromEdges = InPatchWorldDimensions * min(PatchUVCoordinates - InEdgeUVDeadBorder, 1 - InEdgeUVDeadBorder - PatchUVCoordinates);
+		bool bIsInsidePatch = all(DistancesFromEdges >= -FLOAT_TOLERANCE);
+		AlphaT = bIsInsidePatch ? 0 : 1;
+		
+		if (bIsInsidePatch && ClampedFalloffMargin > 0)
+		{
+			// Interpret distances as proportions of falloff margin
+			DistancesFromEdges = DistancesFromEdges / ClampedFalloffMargin;
+		
+			if (DistancesFromEdges.x > 0 && DistancesFromEdges.x < 1.0
+				&& DistancesFromEdges.y > 0 && DistancesFromEdges.y < 1.0)
+			{
+				AlphaT = min(1.0, length(float2(1, 1) - DistancesFromEdges));
+			}
+			else
+			{
+				AlphaT = 1 - min(1.0, min(DistancesFromEdges.x, DistancesFromEdges.y));
+			}
+		}
+	}
+	else
+	{
+		float OuterRadius = min(InPatchWorldDimensions.x, InPatchWorldDimensions.y) / 2;
+		float DistanceFromCenter = length(InPatchWorldDimensions * (PatchUVCoordinates - float2(0.5, 0.5)));
+		float DistanceFromOuterCircle = OuterRadius - DistanceFromCenter;
+		
+		bool bIsInsideCircle = DistanceFromOuterCircle > 0;
+		AlphaT = bIsInsideCircle ? 0 : 1;
+		if (bIsInsideCircle && ClampedFalloffMargin > 0)
+		{
+			AlphaT = 1 - DistanceFromOuterCircle / ClampedFalloffMargin;
+			AlphaT = max(0, AlphaT); // clamp to 0 so we don't go past ClampedFalloffMargin
+		}
+	}
+	
+	float Alpha = cos(AlphaT * HALF_PI);
+	Alpha = Alpha * Alpha * bFalloffAlpha;
+	
+	if (bApplyPatchAlpha)
+	{
+		Alpha *= PatchSampledValue.a;
+	}
+	
+	float NewHeight = 0;
+	if (bAdditiveMode)
+	{
+		NewHeight = CurrentHeight + Alpha * PatchHeight;
+	}
+	else
+	{
+		NewHeight = lerp(CurrentHeight, PatchHeight, Alpha);
+	}
+	
+	OutColor = PackHeight(NewHeight);
+}
+#endif // APPLY_LANDSCAPE_PATCH
+
+#if OFFSET_LANDSCAPE_PATCH
+Texture2D<float4> InHeightmap;
+float InHeightOffset;
+
+void ApplyOffsetToHeightmap(in float4 SVPos : SV_POSITION, out float2 OutColor : SV_Target0)
+{
+	float StoredHeight = UnpackHeight(InHeightmap.Load(int3(SVPos.xy, 0)));
+	OutColor = PackHeight(StoredHeight + InHeightOffset);
+}
+#endif // OFFSET_LANDSCAPE_PATCH
+
+#if SIMPLE_TEXTURE_COPY
+Texture2D<float4> InSource;
+SamplerState InSourceSampler;
+float2 InDestinationResolution;
+
+void SimpleTextureCopy(in float4 SVPos : SV_POSITION, out float2 OutColor : SV_Target0)
+{
+	float2 UV = SVPos.xy / InDestinationResolution;
+	OutColor = InSource.Sample(InSourceSampler, UV);
+}
+#endif // SIMPLE_TEXTURE_COPY
\ No newline at end of file
diff --git a/Engine/Plugins/Experimental/LandscapePatch/Source/LandscapePatch/Private/LandscapePatch.cpp b/Engine/Plugins/Experimental/LandscapePatch/Source/LandscapePatch/Private/LandscapePatch.cpp
deleted file mode 100644
diff --git a/Engine/Plugins/Experimental/LandscapePatch/Source/LandscapePatch/Private/LandscapePatchComponent.cpp b/Engine/Plugins/Experimental/LandscapePatch/Source/LandscapePatch/Private/LandscapePatchComponent.cpp
new file mode 100644
--- /dev/null
+++ b/Engine/Plugins/Experimental/LandscapePatch/Source/LandscapePatch/Private/LandscapePatchComponent.cpp
@@ -0,0 +1,244 @@
+// Copyright Epic Games, Inc. All Rights Reserved.
+
+#include "LandscapePatchComponent.h"
+
+#include "EngineUtils.h"
+#include "GameFramework/Actor.h"
+#include "Landscape.h"
+#include "LandscapePatchLogging.h"
+#include "LandscapePatchManager.h"
+
+#define LOCTEXT_NAMESPACE "LandscapePatch"
+
+#if WITH_EDITOR
+namespace LandscapePatchComponentLocals
+{
+	ALandscapePatchManager* CreateNewPatchManagerForLandscape(ALandscape* Landscape)
+	{
+		if (!ensure(Landscape->CanHaveLayersContent()))
+		{
+			return nullptr;
+		}
+
+		FString BrushActorString = FString::Format(TEXT("{0}_{1}"), { Landscape->GetActorLabel(), ALandscapePatchManager::StaticClass()->GetName() });
+		FName BrushActorName = MakeUniqueObjectName(Landscape->GetOuter(), ALandscapePatchManager::StaticClass(), FName(BrushActorString));
+		FActorSpawnParameters SpawnParams;
+		SpawnParams.Name = BrushActorName;
+		SpawnParams.bAllowDuringConstructionScript = true; // This can be called by construction script if the actor being added to the world is part of a blueprint
+
+		ALandscapePatchManager* PatchManager = Landscape->GetWorld()->SpawnActor<ALandscapePatchManager>(ALandscapePatchManager::StaticClass(), SpawnParams);
+		if (PatchManager)
+		{
+			PatchManager->SetActorLabel(BrushActorString);
+			PatchManager->SetTargetLandscape(Landscape);
+		}
+
+		return PatchManager;
+	}
+}
+
+ULandscapePatchComponent::ULandscapePatchComponent(const FObjectInitializer& ObjectInitializer)
+	: Super(ObjectInitializer)
+{
+	// Causes OnUpdateTransform to be called when the parent is moved. Note that this is better to do here in the
+	// constructor, otherwise we'd need to do it both in OnComponentCreated and PostLoad.
+	// We could keep this false if we were to register to TransformUpdated, since that gets broadcast either way.
+	// TODO: Currently, neither TransformUpdated nor OnUpdateTransform are triggered when parent's transform is changed
+	bWantsOnUpdateTransform = true;
+}
+
+void ULandscapePatchComponent::OnComponentCreated()
+{
+	using namespace LandscapePatchComponentLocals;
+
+	Super::OnComponentCreated();
+		
+	// Mark whether we're creating from scratch of from a copy
+	bWasCopy = bPropertiesCopiedIndicator;
+	bPropertiesCopiedIndicator = true;
+
+	UWorld* World = GetWorld();
+	AActor* OwningActor = GetOwner();
+
+	if (PatchManager.IsValid())
+	{
+		// If we copied over a patch manager, presumably Landscape should be
+		// copied over as well, but might as well do this to be safe.
+		Landscape = PatchManager->GetOwningLandscape();
+	}
+	else if (World)
+	{
+		if (Landscape.IsValid())
+		{
+			// If we copied over a patch with a landscape but no manager, create manager in that landscape
+			if (Landscape->CanHaveLayersContent())
+			{
+				PatchManager = CreateNewPatchManagerForLandscape(Landscape.Get());
+			}
+			else
+			{
+				UE_LOG(LogLandscapePatch, Warning, TEXT("Landscape target for height patch did not have edit layers enabled. Unable to create patch manager."));
+				Landscape = nullptr;
+			}
+		}
+		else
+		{
+			// Didn't copy over an existing manager or landscape.
+
+			// See if the level has a height patch manager to which we can add ourselves
+			TActorIterator<ALandscapePatchManager> ManagerIterator(World);
+			if (!!ManagerIterator)
+			{
+				PatchManager = *ManagerIterator;
+				Landscape = PatchManager->GetOwningLandscape();
+			}
+			else
+			{
+				// If no existing manager, find some landscape and add a new one.
+				for (TActorIterator<ALandscape> LandscapeIterator(World); LandscapeIterator; ++LandscapeIterator)
+				{
+					if (LandscapeIterator->CanHaveLayersContent())
+					{
+						Landscape = *LandscapeIterator;
+						PatchManager = CreateNewPatchManagerForLandscape(Landscape.Get());
+						break;
+					}
+				}
+				if (!PatchManager.IsValid())
+				{
+					UE_LOG(LogLandscapePatch, Warning, TEXT("Unable to find a landscape with edit layers enabled. Unable to create patch manager."));
+				}
+			}
+		}
+	}
+
+	if (PatchManager.IsValid())
+	{
+		PatchManager->AddPatch(this);
+	}
+}
+
+void ULandscapePatchComponent::PostLoad()
+{
+	Super::PostLoad();
+
+	bPropertiesCopiedIndicator = true;
+	bLoadedButNotYetRegistered = true;
+}
+
+void ULandscapePatchComponent::OnComponentDestroyed(bool bDestroyingHierarchy)
+{
+	if (PatchManager.IsValid())
+	{
+		PatchManager->RemovePatch(this);
+	}
+
+	Super::OnComponentDestroyed(bDestroyingHierarchy);
+}
+
+void ULandscapePatchComponent::OnRegister()
+{
+	Super::OnRegister();
+
+
+	if (bLoadedButNotYetRegistered)
+	{
+		bLoadedButNotYetRegistered = false;
+		return;
+	}
+
+	// TODO: Currently the main reason to invalidate the landscape on registration is to respond
+	// to detail panel changes of the parent's transform. However we may be able to catch a wide
+	// variety of changes this way, so we'll need to see if we can get rid of other invalidations.
+	// Also, we should make the invalidation conditional on whether we actually modify any relevant
+	// properties by having a virtual method that compares and updates a stored hash of them.
+	RequestLandscapeUpdate();
+}
+
+void ULandscapePatchComponent::OnUpdateTransform(EUpdateTransformFlags UpdateTransformFlags, ETeleportType Teleport)
+{
+	Super::OnUpdateTransform(UpdateTransformFlags, Teleport);
+
+	RequestLandscapeUpdate();
+}
+
+void ULandscapePatchComponent::PostEditChangeProperty(FPropertyChangedEvent& PropertyChangedEvent)
+{
+	using namespace LandscapePatchComponentLocals;
+
+	Super::PostEditChangeProperty(PropertyChangedEvent);
+
+	// If we're changing the owning landscape, there's some work we need to do to remove/add ourselves from/to the proper
+	// brush managers.
+	if (PropertyChangedEvent.Property 
+		&& (PropertyChangedEvent.Property->GetFName() == GET_MEMBER_NAME_CHECKED(ULandscapePatchComponent, Landscape)))
+	{
+		if (PatchManager.IsValid())
+		{
+			PatchManager->RemovePatch(this);
+		}
+		PatchManager = nullptr;
+
+		bool bAdded = !Landscape.IsValid();
+		for (TActorIterator<ALandscapePatchManager> ManagersItr(GetWorld()); ManagersItr; ++ManagersItr)
+		{
+			// Exit early if found the relevant manager.
+			if (bAdded)
+			{
+				break;
+			}
+
+			ALandscapePatchManager* Manager = *ManagersItr;
+			if (!ensure(Manager->IsValidLowLevel()))
+			{
+				continue;
+			}
+
+			if (!bAdded && Manager->GetOwningLandscape() == Landscape)
+			{
+				Manager->AddPatch(this);
+				PatchManager = Manager;
+				bAdded = true;
+			}
+		}
+
+		// If we had a non-null landscape and didn't add ourselves to its manager despite iterating through
+		// all managers, then the landscape must not have had a manager. Create one for it.
+		if (!bAdded)
+		{
+			if (Landscape->CanHaveLayersContent())
+			{
+				PatchManager = CreateNewPatchManagerForLandscape(Landscape.Get());
+				PatchManager->AddPatch(this);
+			}
+			else
+			{
+				UE_LOG(LogLandscapePatch, Warning, TEXT("Landscape target for height patch did not have edit layers enabled. Unable to create patch manager."));
+				Landscape = nullptr;
+			}
+		}
+	}
+
+	RequestLandscapeUpdate();
+}
+
+void ULandscapePatchComponent::RequestLandscapeUpdate()
+{
+	if(PatchManager.IsValid())
+	{
+		PatchManager->RequestLandscapeUpdate();
+	}
+}
+
+FTransform ULandscapePatchComponent::GetLandscapeHeightmapCoordsToWorld() const
+{
+	if (PatchManager.IsValid())
+	{
+		return PatchManager->GetHeightmapCoordsToWorld();
+	}
+	return FTransform::Identity;
+}
+
+#endif
+
+#undef LOCTEXT_NAMESPACE
diff --git a/Engine/Plugins/Experimental/LandscapePatch/Source/LandscapePatch/Private/LandscapePatchManager.cpp b/Engine/Plugins/Experimental/LandscapePatch/Source/LandscapePatch/Private/LandscapePatchManager.cpp
--- a/Engine/Plugins/Experimental/LandscapePatch/Source/LandscapePatch/Private/LandscapePatchManager.cpp
+++ b/Engine/Plugins/Experimental/LandscapePatch/Source/LandscapePatch/Private/LandscapePatchManager.cpp
@@ -3,7 +3,7 @@
 #include "LandscapePatchManager.h"
 
 #include "Landscape.h"
-#include "LandscapePatch.h"
+#include "LandscapePatchComponent.h"
 #include "Engine/TextureRenderTarget2D.h"
 #include "LandscapeDataAccess.h"
 #include "RenderGraph.h" // RDG_EVENT_NAME
@@ -38,7 +38,7 @@ void ALandscapePatchManager::Initialize_Native(const FTransform & InLandscapeTra
 	const double HEIGHTMAP_TO_OBJECT_HEIGHT_SCALE = LANDSCAPE_ZSCALE;
 	const double HEIGHTMAP_TO_OBJECT_HEIGHT_OFFSET = -LandscapeDataAccess::MidValue * LANDSCAPE_ZSCALE;
 
-	// S_p: the pixel coordinates are actually the same as xy object-space coordinates becaues one quad is 1 unit,
+	// S_p: the pixel coordinate scale is actually the same as xy object-space coordinates because one quad is 1 unit,
 	// so we only need to scale the height.
 	FVector3d PixelToObjectSpaceScale = FVector3d(
 		1,
@@ -46,10 +46,10 @@ void ALandscapePatchManager::Initialize_Native(const FTransform & InLandscapeTra
 		HEIGHTMAP_TO_OBJECT_HEIGHT_SCALE
 	);
 
-	// T_p: again, only the height needs translation
+	// T_p: the center of the pixel
 	FVector3d PixelToObjectSpaceTranslate = FVector3d(
-		0,
-		0,
+		-0.5,
+		-0.5,
 		HEIGHTMAP_TO_OBJECT_HEIGHT_OFFSET
 	);
 
diff --git a/Engine/Plugins/Experimental/LandscapePatch/Source/LandscapePatch/Private/LandscapeTextureHeightPatch.cpp b/Engine/Plugins/Experimental/LandscapePatch/Source/LandscapePatch/Private/LandscapeTextureHeightPatch.cpp
new file mode 100644
--- /dev/null
+++ b/Engine/Plugins/Experimental/LandscapePatch/Source/LandscapePatch/Private/LandscapeTextureHeightPatch.cpp
@@ -0,0 +1,825 @@
+// Copyright Epic Games, Inc. All Rights Reserved.
+
+#include "LandscapeTextureHeightPatch.h"
+
+#include "Engine/TextureRenderTarget2D.h"
+#include "Landscape.h"
+#include "LandscapeDataAccess.h"
+#include "LandscapePatchLogging.h"
+#include "LandscapePatchManager.h"
+#include "RenderGraph.h" // RDG_EVENT_NAME
+#include "TextureCompiler.h"
+#include "UObject/ObjectSaveContext.h"
+
+#define LOCTEXT_NAMESPACE "LandscapeTextureHeightPatch"
+
+#if WITH_EDITOR
+namespace LandscapeTextureHeightPatchLocals
+{
+	static const FLinearColor ClearColor = LandscapeDataAccess::PackHeight(LandscapeDataAccess::MidValue).ReinterpretAsLinear();
+
+	void CopyTextureOnRenderThread(FRHICommandListImmediate& RHICmdList, const FTextureResource& Source, FTextureResource& Destination)
+	{
+		using namespace UE::Landscape;
+
+		FRDGBuilder GraphBuilder(RHICmdList, RDG_EVENT_NAME("LandscapeTextureHeightPatchCopyTexture"));
+
+		FRDGTextureRef SourceTexture = GraphBuilder.RegisterExternalTexture(CreateRenderTarget(Source.GetTexture2DRHI(), TEXT("CopySource")));
+		FRDGTextureRef DestinationTexture = GraphBuilder.RegisterExternalTexture(CreateRenderTarget(Destination.GetTexture2DRHI(), TEXT("CopyDestination")));
+
+		// All my efforts of getting CopyToResolveTarget to work without complaints have failed, so we just use our own copy shader.
+		FSimpleTextureCopyPS::AddToRenderGraph(GraphBuilder, SourceTexture, DestinationTexture);
+
+		GraphBuilder.Execute();
+	}
+}
+
+void ULandscapeTextureHeightPatch::CopyInternalTextureToInternalRenderTarget()
+{
+	using namespace LandscapeTextureHeightPatchLocals;
+
+	if (IsValid(InternalTexture))
+	{
+		InternalTexture->UpdateResource();
+	}
+
+	// TODO: should the GetResource be an ensure?
+	if (!IsValid(InternalTexture) || !InternalTexture->GetResource())
+	{
+		if (InternalRenderTarget)
+		{
+			Modify();
+			InternalRenderTarget = nullptr;
+		}
+		return;
+	}
+
+	if (IsValid(InternalRenderTarget))
+	{
+		InternalRenderTarget->Modify();
+	}
+	ResizeRenderTargetIfNeeded(InternalTexture->GetResource()->GetSizeX(), 
+		InternalTexture->GetResource()->GetSizeY());
+
+	FTextureCompilingManager::Get().FinishCompilation({ InternalTexture });
+
+	FTextureResource* Source = InternalTexture->GetResource();
+	FTextureResource* Destination = InternalRenderTarget->GetResource();
+	ENQUEUE_RENDER_COMMAND(LandscapeTextureHeightPatchRTToTexture)(
+		[Source, Destination](FRHICommandListImmediate& RHICmdList)
+	{
+		CopyTextureOnRenderThread(RHICmdList, *Source, *Destination);
+	});
+
+}
+
+void ULandscapeTextureHeightPatch::CopyInternalRenderTargetToInternalTexture(bool bBlock)
+{
+	using namespace LandscapeTextureHeightPatchLocals;
+
+	if (!IsValid(InternalRenderTarget))
+	{
+		if (InternalTexture)
+		{
+			Modify();
+			InternalTexture = nullptr;
+		}
+		return;
+	}
+
+	if (IsValid(InternalTexture))
+	{
+		InternalTexture->Modify();
+	}
+	ResizeTextureIfNeeded(InternalRenderTarget->SizeX, InternalRenderTarget->SizeY, /*bClear*/ false, /*bUpdateResource*/ false);
+
+	// Read from the render target.
+	// TODO: The header for this requires the texture to be square power of 2, but it actually doesn't seem to
+	// be an enforced requirement. If that changes, we'll need our own ReadPixels followed by locking a mip
+	// and writing to it.
+	InternalRenderTarget->UpdateTexture2D(InternalTexture, ETextureSourceFormat::TSF_BGRA8);
+
+	InternalTexture->UpdateResource();
+	
+	if (bBlock)
+	{
+		FTextureCompilingManager::Get().FinishCompilation({ InternalTexture });
+	}
+}
+
+// Render targets don't get serialized, so whenever we need to save, copy, etc, we convert
+// to a UTexture2D, and then we convert back when needed.
+void ULandscapeTextureHeightPatch::PostLoad()
+{
+	Super::PostLoad();
+
+	if (SourceMode == ELandscapeTexturePatchSourceMode::TextureBackedRenderTarget)
+	{
+		CopyInternalTextureToInternalRenderTarget();
+	}
+}
+
+void ULandscapeTextureHeightPatch::PreSave(FObjectPreSaveContext SaveContext)
+{
+	Super::PreSave(SaveContext);
+
+	if (SourceMode == ELandscapeTexturePatchSourceMode::TextureBackedRenderTarget)
+	{
+		CopyInternalRenderTargetToInternalTexture(true);
+	}
+}
+
+void ULandscapeTextureHeightPatch::PreDuplicate(FObjectDuplicationParameters& DupParams)
+{
+	Super::PreDuplicate(DupParams);
+
+	if (SourceMode == ELandscapeTexturePatchSourceMode::TextureBackedRenderTarget)
+	{
+		CopyInternalRenderTargetToInternalTexture(true);
+	}
+}
+
+// Called when serializing to text for copy/paste
+void ULandscapeTextureHeightPatch::ExportCustomProperties(FOutputDevice& Out, uint32 Indent)
+{
+	Super::ExportCustomProperties(Out, Indent);
+
+	if (SourceMode == ELandscapeTexturePatchSourceMode::TextureBackedRenderTarget)
+	{
+		CopyInternalRenderTargetToInternalTexture(true);
+	}
+}
+
+#if WITH_EDITOR
+void ULandscapeTextureHeightPatch::PostEditChangeProperty(FPropertyChangedEvent& PropertyChangedEvent)
+{
+	if (PropertyChangedEvent.Property)
+	{
+		if (PropertyChangedEvent.Property->GetFName() == GET_MEMBER_NAME_CHECKED(ULandscapeTextureHeightPatch, SourceMode))
+		{
+			SetSourceMode(SourceMode);
+		}
+		else if (PropertyChangedEvent.Property->GetFName() == GET_MEMBER_NAME_CHECKED(ULandscapeTextureHeightPatch, InitializationMode))
+		{
+			bShowTextureAssetProperty = InitializationMode == ELandscapeTextureHeightPatchInitMode::TextureAsset;
+		}
+	}
+		
+	Super::PostEditChangeProperty(PropertyChangedEvent);
+}
+
+void ULandscapeTextureHeightPatch::PostInitProperties()
+{
+	bShowTextureAssetProperty = InitializationMode == ELandscapeTextureHeightPatchInitMode::TextureAsset;
+
+	Super::PostInitProperties();
+}
+#endif
+
+void ULandscapeTextureHeightPatch::OnComponentCreated()
+{
+	using namespace LandscapeTextureHeightPatchLocals;
+
+	Super::OnComponentCreated();
+
+	if (bWasCopy)
+	{
+		if (SourceMode == ELandscapeTexturePatchSourceMode::TextureBackedRenderTarget
+			&& IsValid(InternalTexture))
+		{
+			CopyInternalTextureToInternalRenderTarget();
+		}
+	}
+	else // if not copy, ie adding a totally new component
+	{
+		// Set component location to be the bottom of the parent actor bounding box.
+		AActor* ParentActor = GetAttachParentActor();
+		if (ParentActor)
+		{
+			FVector Origin, BoxExtent;
+			GetAttachParentActor()->GetActorBounds(false, Origin, BoxExtent);
+
+			Origin.Z -= BoxExtent.Z;
+			SetWorldLocation(Origin);
+		}
+	}
+
+	PreviousSourceMode = SourceMode;
+}
+
+void ULandscapeTextureHeightPatch::UpdateShaderParams(
+	UE::Landscape::FApplyLandscapeTextureHeightPatchPS::FParameters& Params, 
+	const FIntPoint& DestinationResolution, FIntRect& DestinationBoundsOut) const
+{
+	using namespace UE::Geometry;
+
+	// We want our patch to be oriented with its Z axis to be along the Z axis of the landscape. The way we do this here
+	// is by just changing the rotation component of the patch transform to be the rotation of the landscape, except for
+	// yaw (rotation around Z axis) that we copy from the patch transform.
+	// Obviously this is not perfect, but it's not clear whether there's a clean way to deal with differing transforms for
+	// landscape vs patch.
+	// TODO: Make it so that patches do not inherit parent non-yaw components of rotation?
+
+	FTransform PatchToWorld = GetPatchToWorldTransform(); 
+	
+	FVector2D FullPatchDimensions = GetFullUnscaledWorldSize();
+	Params.InPatchWorldDimensions = FVector2f(FullPatchDimensions);
+
+	FTransform FromPatchUVToPatch(FQuat4d::Identity, FVector3d(-FullPatchDimensions.X / 2, -FullPatchDimensions.Y / 2, 0), 
+		FVector3d(FullPatchDimensions.X, FullPatchDimensions.Y, 1));
+	FMatrix44d PatchLocalToUVs = FromPatchUVToPatch.ToInverseMatrixWithScale();
+
+	FTransform LandscapeHeightmapToWorld = PatchManager->GetHeightmapCoordsToWorld();
+	FMatrix44d LandscapeToWorld = LandscapeHeightmapToWorld.ToMatrixWithScale();
+
+	FMatrix44d WorldToPatch = PatchToWorld.ToInverseMatrixWithScale();
+
+	// In unreal, matrix composition is done by multiplying the subsequent ones on the right, and the result
+	// is transpose of what our shader will expect (because unreal right multiplies vectors by matrices).
+	FMatrix44d LandscapeToPatchUVTransposed = LandscapeToWorld * WorldToPatch * PatchLocalToUVs;
+	Params.InHeightmapToPatch = (FMatrix44f)LandscapeToPatchUVTransposed.GetTransposed();
+
+	FVector3d ComponentScale = PatchToWorld.GetScale3D();
+	float CombinedHeightScale = PatchHeightScale * (bApplyComponentZScale ? ComponentScale.Z : 1);
+	Params.InHeightScale = CombinedHeightScale;
+
+	Params.InPatchBaseHeight = 0;
+	if (bUsePatchZAsReference)
+	{
+		FVector3d PatchOriginInHeightmapCoords = LandscapeHeightmapToWorld.InverseTransformPosition(PatchToWorld.GetTranslation());
+
+		Params.InPatchBaseHeight = CombinedHeightScale < 0 ?
+			PatchOriginInHeightmapCoords.Z - (LandscapeDataAccess::MaxValue - LandscapeDataAccess::MidValue) * CombinedHeightScale
+			: PatchOriginInHeightmapCoords.Z - LandscapeDataAccess::MidValue * CombinedHeightScale;
+	}
+
+	// The outer half-pixel shouldn't affect the landscape because it is not part of our official coverage area.
+	Params.InEdgeUVDeadBorder = FVector2f::Zero();
+	FVector2D TextureResolution;
+	if (GetTextureResolution(TextureResolution))
+	{
+		Params.InEdgeUVDeadBorder = FVector2f(0.5 / TextureResolution.X, 0.5 / TextureResolution.Y);
+	}
+
+	Params.bRectangularFalloff = FalloffMode == ELandscapeTextureHeightPatchFalloffMode::RoundedRectangle;
+	Params.InFalloffWorldMargin = Falloff / FMath::Min(ComponentScale.X, ComponentScale.Y);
+
+	Params.bApplyPatchAlpha = bUseTextureAlphaChannel;
+	Params.bAdditiveMode = BlendMode == ELandscapeTextureHeightPatchBlendMode::Additive;
+
+	// Get the output bounds, which are used to limit the amount of landscape pixels we have to process. 
+	// To get them, convert all of the corners into heightmap 2d coordinates and get the bounding box.
+	auto PatchUVToHeightmap2DCoordinates = [&PatchToWorld, &FromPatchUVToPatch, &LandscapeHeightmapToWorld](const FVector2f& UV)
+	{
+		FVector WorldPosition = PatchToWorld.TransformPosition(
+			FromPatchUVToPatch.TransformPosition(FVector(UV.X, UV.Y, 0)));
+		FVector HeightmapCoordinates = LandscapeHeightmapToWorld.InverseTransformPosition(WorldPosition);
+		return FVector2d(HeightmapCoordinates.X, HeightmapCoordinates.Y);
+	};
+	FBox2D FloatBounds(ForceInit);
+	FloatBounds += PatchUVToHeightmap2DCoordinates(FVector2f(0, 0));
+	FloatBounds += PatchUVToHeightmap2DCoordinates(FVector2f(0,1));
+	FloatBounds += PatchUVToHeightmap2DCoordinates(FVector2f(1, 0));
+	FloatBounds += PatchUVToHeightmap2DCoordinates(FVector2f(1,1));
+
+	DestinationBoundsOut = FIntRect(
+		FMath::Clamp(FMath::Floor(FloatBounds.Min.X), 0, DestinationResolution.X - 1),
+		FMath::Clamp(FMath::Floor(FloatBounds.Min.Y), 0, DestinationResolution.Y - 1),
+		FMath::Clamp(FMath::CeilToInt(FloatBounds.Max.X) + 1, 0, DestinationResolution.X),
+		FMath::Clamp(FMath::CeilToInt(FloatBounds.Max.Y) + 1, 0, DestinationResolution.Y));
+}
+
+UTextureRenderTarget2D* ULandscapeTextureHeightPatch::Render_Native(bool bIsHeightmap,
+	UTextureRenderTarget2D* InCombinedResult,
+	const FName& InWeightmapLayerName)
+{
+	using namespace UE::Landscape;
+
+	if (!ensure(PatchManager.IsValid()) || 
+		(SourceMode == ELandscapeTexturePatchSourceMode::TextureBackedRenderTarget && (!IsValid(InternalRenderTarget)))
+		|| (SourceMode == ELandscapeTexturePatchSourceMode::TextureAsset && (!IsValid(TextureAsset)))
+		|| (SourceMode == ELandscapeTexturePatchSourceMode::InternalTexture && (!IsValid(InternalTexture))))
+	{
+		return InCombinedResult;
+	}
+
+	// Texture height patch doesn't affect weight maps
+	if (!bIsHeightmap)
+	{
+		return InCombinedResult;
+	}
+
+	// Go ahead and pack everything into a copy of the param struct so we don't have to capture everything
+	// individually in the lambda below.
+	FApplyLandscapeTextureHeightPatchPS::FParameters ShaderParamsToCopy;
+	FIntRect DestinationBounds;
+	UpdateShaderParams(ShaderParamsToCopy, FIntPoint(InCombinedResult->SizeX, InCombinedResult->SizeY), DestinationBounds);
+
+	if (DestinationBounds.IsEmpty())
+	{
+		// Patch must be outside the landscape.
+		return InCombinedResult;
+	}
+
+	UTexture* PatchUObject = InternalTexture;
+	if (SourceMode == ELandscapeTexturePatchSourceMode::TextureBackedRenderTarget)
+	{
+		PatchUObject = InternalRenderTarget;
+	}
+	else if (SourceMode == ELandscapeTexturePatchSourceMode::TextureAsset)
+	{
+		PatchUObject = TextureAsset;
+	} 
+	
+	if (!PatchUObject)
+	{
+		return InCombinedResult;
+	}
+	FTextureResource* Patch = PatchUObject->GetResource();
+
+	ENQUEUE_RENDER_COMMAND(LandscapeTextureHeightPatch)([InCombinedResult, ShaderParamsToCopy, Patch, DestinationBounds](FRHICommandListImmediate& RHICmdList)
+	{
+		TRACE_CPUPROFILER_EVENT_SCOPE(LandscapeTextureHeightPatch_Render);
+
+		FRDGBuilder GraphBuilder(RHICmdList, RDG_EVENT_NAME("ApplyLandmassTextureHeightPatch"));
+		
+		TRefCountPtr<IPooledRenderTarget> DestinationRenderTarget = CreateRenderTarget(InCombinedResult->GetResource()->GetTexture2DRHI(), TEXT("LandscapeTextureHeightPatchOutput"));
+		FRDGTextureRef DestinationTexture = GraphBuilder.RegisterExternalTexture(DestinationRenderTarget);
+
+		// Make a copy of our heightmap input so we can read and write at the same time (needed for blending)
+		FRDGTextureRef InputCopy = GraphBuilder.CreateTexture(DestinationTexture->Desc, TEXT("LandscapeTextureHeightPatchInputCopy"));
+
+		FRHICopyTextureInfo CopyTextureInfo;
+		CopyTextureInfo.NumMips = 1;
+		CopyTextureInfo.Size = FIntVector(DestinationTexture->Desc.GetSize().X, DestinationTexture->Desc.GetSize().Y, 0);
+		AddCopyTexturePass(GraphBuilder, DestinationTexture, InputCopy, CopyTextureInfo);
+
+		FApplyLandscapeTextureHeightPatchPS::FParameters* ShaderParams =
+			GraphBuilder.AllocParameters<FApplyLandscapeTextureHeightPatchPS::FParameters>();
+		*ShaderParams = ShaderParamsToCopy;
+
+		TRefCountPtr<IPooledRenderTarget> PatchRenderTarget = CreateRenderTarget(Patch->GetTexture2DRHI(), TEXT("LandscapeTextureHeightPatch"));
+		FRDGTextureRef PatchTexture = GraphBuilder.RegisterExternalTexture(PatchRenderTarget);
+		FRDGTextureSRVRef PatchSRV = GraphBuilder.CreateSRV(FRDGTextureSRVDesc::CreateForMipLevel(PatchTexture, 0));
+		ShaderParams->InHeightPatch = PatchSRV;
+		ShaderParams->InHeightPatchSampler = TStaticSamplerState<SF_Bilinear, AM_Clamp, AM_Clamp>::GetRHI();
+
+		FRDGTextureSRVRef InputCopySRV = GraphBuilder.CreateSRV(FRDGTextureSRVDesc::CreateForMipLevel(InputCopy, 0));
+		ShaderParams->InSourceHeightmap = InputCopySRV;
+
+		ShaderParams->RenderTargets[0] = FRenderTargetBinding(DestinationTexture, ERenderTargetLoadAction::ENoAction, /*InMipIndex = */0);
+
+		FApplyLandscapeTextureHeightPatchPS::AddToRenderGraph(GraphBuilder, ShaderParams, DestinationBounds);
+
+		GraphBuilder.Execute();
+	});
+
+	return InCombinedResult;
+}
+
+void ULandscapeTextureHeightPatch::DeleteInternalTextures()
+{
+	if (InternalRenderTarget || InternalTexture)
+	{
+		Modify();
+	}
+
+	InternalRenderTarget = nullptr;
+	InternalTexture = nullptr;
+
+	if (SourceMode != ELandscapeTexturePatchSourceMode::TextureAsset && PatchManager.IsValid())
+	{
+		PatchManager->RequestLandscapeUpdate();
+	}
+}
+
+void ULandscapeTextureHeightPatch::Reinitialize()
+{
+	using namespace UE::Landscape;
+	using namespace LandscapeTextureHeightPatchLocals;
+
+	if (SourceMode == ELandscapeTexturePatchSourceMode::TextureAsset)
+	{
+		// Nothing to reinitialize...
+		UE_LOG(LogLandscapePatch, Warning, TEXT("ULandscapeTextureHeightPatch::Reinitialize: Unable to reinitialize when source mode is set"
+			"to be a texture asset."));
+		return;
+	}
+	if ((InitializationMode == ELandscapeTextureHeightPatchInitMode::TextureAsset && !IsValid(TextureAsset))
+		|| (InitializationMode == ELandscapeTextureHeightPatchInitMode::FromLandscape && !Landscape.IsValid()))
+	{
+		// Don't have what we need for initialization
+		return;
+	}
+
+	// Figure out what needs to be modified for undo/redo...
+	if (SourceMode == ELandscapeTexturePatchSourceMode::TextureBackedRenderTarget)
+	{
+		if (!IsValid(InternalRenderTarget))
+		{
+			// A new render target is going to be created
+			Modify();
+		}
+		else
+		{
+			InternalRenderTarget->Modify();
+		}
+	}
+	else if (SourceMode == ELandscapeTexturePatchSourceMode::InternalTexture)
+	{
+		if (!IsValid(InternalTexture))
+		{
+			// A new texture is going to be created
+			Modify();
+		}
+		else
+		{
+			InternalTexture->Modify();
+		}
+	}
+
+
+	switch (InitializationMode)
+	{
+	case ELandscapeTextureHeightPatchInitMode::Blank:
+		if (SourceMode == ELandscapeTexturePatchSourceMode::TextureBackedRenderTarget)
+		{
+			ResizeRenderTargetIfNeeded(InitTextureSizeX, InitTextureSizeY);
+		}
+		else if (SourceMode == ELandscapeTexturePatchSourceMode::InternalTexture)
+		{
+			ResizeTextureIfNeeded(InitTextureSizeX, InitTextureSizeY, /*bClear*/ true, /*bUpdateResource*/ true);
+		}
+		break;
+	case ELandscapeTextureHeightPatchInitMode::FromLandscape:
+	{
+		if (bBaseResolutionOffLandscape)
+		{
+			FVector2D DesiredResolution;
+			GetInitResolutionFromLandscape(ResolutionMultiplier, DesiredResolution);
+			InitTextureSizeX = DesiredResolution.X;
+			InitTextureSizeY = DesiredResolution.Y;
+		}
+
+		// We're going to need our render target regardless of source mode because we'll write to it 
+		// before copying to texture if needed.
+		ResizeRenderTargetIfNeeded(InitTextureSizeX, InitTextureSizeY);
+			
+		// If bUsePatchZAsReference is true, then we're going to be adding an offset to our rendered landscape
+		// section in anticipation of adding an offset based on patch z later. In that case, we'll render
+		// the landscape to an intermediate target first, then do the copy to the internal target while adding
+		// in the offset. Otherwise, we can render directly to internal render target.
+		UTextureRenderTarget2D* RenderedHeightmapSection = InternalRenderTarget;
+		if (bUsePatchZAsReference)
+		{
+			RenderedHeightmapSection = NewObject<UTextureRenderTarget2D>(this);
+			RenderedHeightmapSection->RenderTargetFormat = ETextureRenderTargetFormat::RTF_RG8;
+			RenderedHeightmapSection->InitAutoFormat(InitTextureSizeX, InitTextureSizeY);
+			RenderedHeightmapSection->UpdateResourceImmediate(true);
+		}
+
+		// Note that we need the "for resolution" overload here since our current texture may not yet have the correct resolution
+		FVector2d FullPatchDimensions(GetFullUnscaledWorldSizeForResolution(FVector2D(InitTextureSizeX, InitTextureSizeY)));
+
+		Landscape->RenderHeightmap(GetPatchToWorldTransform(), FBox2D(-FullPatchDimensions / 2, FullPatchDimensions / 2), RenderedHeightmapSection);
+
+		// Reset height related parameters because adjusting for them would be a pain and doesn't
+		// seem useful to the user if they are reinitializing...
+		bApplyComponentZScale = false;
+		PatchHeightScale = 1;
+
+		// ...However, we do need to apply the offset to account for our patch's position relative to the
+		// landscape z axis, if that is a setting that the user is using.
+		if (bUsePatchZAsReference)
+		{
+			FTransform LandscapeHeightmapToWorld = PatchManager->GetHeightmapCoordsToWorld();
+			FVector3d PatchOriginInHeightmapCoords = LandscapeHeightmapToWorld.InverseTransformPosition(GetComponentTransform().GetTranslation());
+			float OffsetToApply = LandscapeDataAccess::MidValue - PatchOriginInHeightmapCoords.Z;
+
+			ENQUEUE_RENDER_COMMAND(RenderHeightmap)([Input = RenderedHeightmapSection->GetResource(), Patch = InternalRenderTarget->GetResource(), OffsetToApply](FRHICommandListImmediate& RHICmdList)
+			{
+				FRDGBuilder GraphBuilder(RHICmdList, RDG_EVENT_NAME("OffsetTextureHeightPatchInitialization"));
+
+				TRefCountPtr<IPooledRenderTarget> InputRenderTarget = CreateRenderTarget(Input->GetTexture2DRHI(), TEXT("LandscapeTextureHeightPatchOffsetOutput"));
+				FRDGTextureRef InputTexture = GraphBuilder.RegisterExternalTexture(InputRenderTarget);
+
+				TRefCountPtr<IPooledRenderTarget> DestinationRenderTarget = CreateRenderTarget(Patch->GetTexture2DRHI(), TEXT("LandscapeTextureHeightPatchOffsetOutput"));
+				FRDGTextureRef DestinationTexture = GraphBuilder.RegisterExternalTexture(DestinationRenderTarget);
+
+				FOffsetHeightmapPS::FParameters* ShaderParams =
+					GraphBuilder.AllocParameters<FOffsetHeightmapPS::FParameters>();
+
+				FRDGTextureSRVRef InputSRV = GraphBuilder.CreateSRV(FRDGTextureSRVDesc::CreateForMipLevel(InputTexture, 0));
+				ShaderParams->InHeightmap = InputSRV;
+				ShaderParams->InHeightOffset = OffsetToApply;
+				ShaderParams->RenderTargets[0] = FRenderTargetBinding(DestinationTexture, ERenderTargetLoadAction::ENoAction, /*InMipIndex = */0);
+
+				FOffsetHeightmapPS::AddToRenderGraph(GraphBuilder, ShaderParams);
+
+				GraphBuilder.Execute();
+			});
+		}
+
+		if (SourceMode == ELandscapeTexturePatchSourceMode::InternalTexture)
+		{
+			CopyInternalRenderTargetToInternalTexture(true);
+			InternalRenderTarget = nullptr;
+		}
+
+		break;
+	}
+	case ELandscapeTextureHeightPatchInitMode::TextureAsset:
+	{
+		if (!IsValid(TextureAsset))
+		{
+			break;
+		}
+
+		if (!TextureAsset->GetResource())
+		{
+			TextureAsset->UpdateResource();
+			FTextureCompilingManager::Get().FinishCompilation({ TextureAsset });
+		}
+
+		if (bBaseResolutionOffLandscape)
+		{
+			InitTextureSizeX = TextureAsset->GetResource()->GetSizeX();
+			InitTextureSizeY = TextureAsset->GetResource()->GetSizeY();
+		}
+
+		// We're going to need our render target regardless of source mode because we'll write to it 
+		// before copying to texture if needed.
+		ResizeRenderTargetIfNeeded(InitTextureSizeX, InitTextureSizeY);
+
+		ENQUEUE_RENDER_COMMAND(LandscapeTextureHeightPatchCopyInternalRT)(
+			[Source = TextureAsset->GetResource(), Destination = InternalRenderTarget->GetResource()](FRHICommandListImmediate& RHICmdList)
+		{
+			CopyTextureOnRenderThread(RHICmdList, *Source, *Destination);
+		});
+
+		if (SourceMode == ELandscapeTexturePatchSourceMode::InternalTexture)
+		{
+			CopyInternalRenderTargetToInternalTexture(true);
+			InternalRenderTarget = nullptr;
+		}
+		break;
+	}
+	}//end switch on InitializationMode
+
+	if (PatchManager.IsValid())
+	{
+		PatchManager->RequestLandscapeUpdate();
+	}
+}
+
+bool ULandscapeTextureHeightPatch::ResizeRenderTargetIfNeeded(int32 SizeX, int32 SizeY)
+{
+	using namespace LandscapeTextureHeightPatchLocals;
+
+	bool bChanged = false;
+
+	if (!IsValid(InternalRenderTarget))
+	{
+		Modify();
+
+		InternalRenderTarget = NewObject<UTextureRenderTarget2D>(this);
+		InternalRenderTarget->ClearColor = ClearColor;
+
+		InternalRenderTarget->RenderTargetFormat = ETextureRenderTargetFormat::RTF_RGBA8;
+		InternalRenderTarget->InitAutoFormat(SizeX, SizeY);
+
+		bChanged = true;
+	}
+	else if (InternalRenderTarget->SizeX != SizeX
+		|| InternalRenderTarget->SizeY != SizeY)
+	{
+		InternalRenderTarget->Modify();
+		InternalRenderTarget->RenderTargetFormat = ETextureRenderTargetFormat::RTF_RGBA8;
+		InternalRenderTarget->InitAutoFormat(SizeX, SizeY);
+
+		bChanged = true;
+	}
+
+	InternalRenderTarget->UpdateResourceImmediate(true);
+
+	return bChanged;
+}
+
+bool ULandscapeTextureHeightPatch::ResizeTextureIfNeeded(int32 SizeX, int32 SizeY, bool bClear, bool bUpdateResource)
+{
+	using namespace LandscapeTextureHeightPatchLocals;
+
+	bool bChanged = false;
+
+	if (!IsValid(InternalTexture) || !InternalTexture->GetResource())
+	{
+		Modify();
+		bChanged = true;
+
+		InternalTexture = NewObject<UTexture2D>(this);
+		InternalTexture->SRGB = false;
+		InternalTexture->MipGenSettings = TMGS_NoMipmaps;
+		InternalTexture->AddressX = TA_Clamp;
+		InternalTexture->AddressY = TA_Clamp;
+		InternalTexture->LODGroup = TEXTUREGROUP_Terrain_Heightmap;
+		// TODO: How do we allow lossless compression, instead of disallowing compression entirely? Just setting 
+		// LossyCompressionAmount to ETextureLossyCompressionAmount::TLCA_None is not sufficient.
+		InternalTexture->CompressionNone = true;
+	}
+	else if (bClear
+		|| InternalTexture->GetResource()->GetSizeX() != SizeX
+		|| InternalTexture->GetResource()->GetSizeY() != SizeY)
+	{
+		bChanged = true;
+		InternalTexture->Modify();
+	}
+
+	if (!InternalTexture->GetResource()
+		|| InternalTexture->GetResource()->GetSizeX() != SizeX
+		|| InternalTexture->GetResource()->GetSizeY() != SizeY)
+	{
+		InternalTexture->Source.Init(SizeX, SizeY, 1, 1, ETextureSourceFormat::TSF_BGRA8);
+	}
+	if (bClear)
+	{
+		TArray<FColor> SourceColors;
+		SourceColors.Init(LandscapeDataAccess::PackHeight(LandscapeDataAccess::MidValue), SizeX * SizeY);
+		uint8* SourceData = InternalTexture->Source.LockMip(0);
+		FMemory::Memcpy(SourceData, SourceColors.GetData(), sizeof(FColor) * SourceColors.Num());
+		InternalTexture->Source.UnlockMip(0);
+	}
+
+	if (bChanged)
+	{
+		if (bUpdateResource)
+		{
+			InternalTexture->UpdateResource();
+		}
+
+		// TODO: Do we need these?
+		InternalTexture->PostEditChange();
+		InternalTexture->MarkPackageDirty();
+	}
+
+	return bChanged;
+}
+
+bool ULandscapeTextureHeightPatch::SetSourceMode(ELandscapeTexturePatchSourceMode NewMode, bool bDeleteUnusedInternalTextures)
+{
+	SourceMode = NewMode;
+
+	if (PreviousSourceMode == NewMode)
+	{
+		return true;
+	}
+
+	if (PreviousSourceMode == ELandscapeTexturePatchSourceMode::TextureBackedRenderTarget
+		&& NewMode == ELandscapeTexturePatchSourceMode::InternalTexture)
+	{
+		CopyInternalRenderTargetToInternalTexture(false);
+	}
+	else if (PreviousSourceMode == ELandscapeTexturePatchSourceMode::InternalTexture
+		&& NewMode == ELandscapeTexturePatchSourceMode::TextureBackedRenderTarget)
+	{
+		CopyInternalTextureToInternalRenderTarget();
+	}
+
+	if (bDeleteUnusedInternalTextures)
+	{
+		if (NewMode != ELandscapeTexturePatchSourceMode::TextureBackedRenderTarget)
+		{
+			InternalRenderTarget = nullptr;
+		}
+		if (NewMode == ELandscapeTexturePatchSourceMode::TextureAsset)
+		{
+			InternalTexture = nullptr;
+		}
+	}
+
+	PreviousSourceMode = NewMode;
+
+	return true;
+}
+
+void ULandscapeTextureHeightPatch::SnapToLandscape()
+{
+	Modify();
+
+	if (!Landscape.IsValid())
+	{
+		return;
+	}
+
+	FTransform LandscapeTransform = Landscape->GetTransform();
+	FTransform PatchTransform = GetComponentTransform();
+
+	FQuat LandscapeRotation = LandscapeTransform.GetRotation();
+	FQuat PatchRotation = PatchTransform.GetRotation();
+
+	// Get rotation of patch relative to landscape
+	FQuat PatchRotationRelativeLandscape = LandscapeRotation.Inverse()* PatchRotation;
+
+	// Get component of that relative rotation that is around the landscape Z axis.
+	double RadiansAroundZ = PatchRotationRelativeLandscape.GetTwistAngle(FVector::ZAxisVector);
+
+	// Round that rotation to nearest 90 degree increment
+	int32 Num90DegreeRotations = FMath::RoundToDouble(RadiansAroundZ / FMathd::HalfPi);
+	double NewRadiansAroundZ = Num90DegreeRotations * FMathd::HalfPi;
+
+	// Now adjust the patch transform.
+	FQuat NewPatchRotation = FQuat(FVector::ZAxisVector, NewRadiansAroundZ) * LandscapeRotation;
+	SetWorldRotation(NewPatchRotation);
+
+	// Once we have the rotation adjusted, we need to adjust the patch size and positioning.
+	// However don't bother if either the patch or landscape scale is 0. We might still be able
+	// to align in one of the axes in such a case, but it is not worth the code complexity for
+	// a broken use case.
+	FVector LandscapeScale = Landscape->GetTransform().GetScale3D();
+	FVector PatchScale = GetComponentTransform().GetScale3D();
+	if (LandscapeScale.X == 0 || LandscapeScale.Y == 0)
+	{
+		UE_LOG(LogLandscapePatch, Warning, TEXT("ULandscapeTextureHeightPatch::SnapToLandscape: Landscape target "
+			"for height patch had a zero scale in one of the dimensions. Skipping aligning position."));
+		return;
+	}
+	if (PatchScale.X == 0 || PatchScale.Y == 0)
+	{
+		UE_LOG(LogLandscapePatch, Warning, TEXT("ULandscapeTextureHeightPatch::SnapToLandscape: Height patch "
+			"had a zero scale in one of the dimensions. Skipping aligning position."));
+		return;
+	}
+
+	// Start by adjusting size to be a multiple of landscape quad size.
+	double PatchExtentX = PatchScale.X * UnscaledPatchCoverage.X;
+	double PatchExtentY = PatchScale.Y * UnscaledPatchCoverage.Y;
+	if (Num90DegreeRotations % 2)
+	{
+		// Relative to the landscape, our lenght and width are backwards...
+		Swap(PatchExtentX, PatchExtentY);
+	}
+
+	int32 LandscapeQuadsX = FMath::RoundToInt(PatchExtentX / LandscapeScale.X);
+	int32 LandscapeQuadsY = FMath::RoundToInt(PatchExtentY / LandscapeScale.Y);
+
+	double NewPatchExtentX = LandscapeQuadsX * LandscapeScale.X;
+	double NewPatchExtentY = LandscapeQuadsY * LandscapeScale.Y;
+	if (Num90DegreeRotations % 2)
+	{
+		Swap(NewPatchExtentX, NewPatchExtentY);
+	}
+	UnscaledPatchCoverage = FVector2D(NewPatchExtentX / PatchScale.X, NewPatchExtentY / PatchScale.Y);
+
+	// Now adjust the center of the patch. This gets snapped to either integer or integer + 0.5 increments
+	// in landscape coordinates depending on whether patch length/width is odd or even in landscape coordinates.
+
+	FVector PatchCenterInLandscapeCoordinates = LandscapeTransform.InverseTransformPosition(GetComponentLocation());
+	double NewPatchCenterX, NewPatchCenterY;
+	if (LandscapeQuadsX % 2)
+	{
+		NewPatchCenterX = FMath::RoundToDouble(PatchCenterInLandscapeCoordinates.X + 0.5) - 0.5;
+	}
+	else
+	{
+		NewPatchCenterX = FMath::RoundToDouble(PatchCenterInLandscapeCoordinates.X);
+	}
+	if (LandscapeQuadsY % 2)
+	{
+		NewPatchCenterY = FMath::RoundToDouble(PatchCenterInLandscapeCoordinates.Y + 0.5) - 0.5;
+	}
+	else
+	{
+		NewPatchCenterY = FMath::RoundToDouble(PatchCenterInLandscapeCoordinates.Y);
+	}
+
+	FVector NewCenterInLandscape(NewPatchCenterX, NewPatchCenterY, PatchCenterInLandscapeCoordinates.Z);
+	SetWorldLocation(LandscapeTransform.TransformPosition(NewCenterInLandscape));
+}
+
+bool ULandscapeTextureHeightPatch::SetTextureResolution(FVector2D ResolutionIn)
+{
+	if (SourceMode == ELandscapeTexturePatchSourceMode::TextureAsset)
+	{
+		UE_LOG(LogLandscapePatch, Warning, TEXT("ULandscapeTextureHeightPatch::ResizeIfNeeded: Cannot resize when source mode is set to texture asset."));
+		return false;
+	}
+
+	if (ResolutionIn.X <= 0 || ResolutionIn.Y <= 0)
+	{
+		UE_LOG(LogLandscapePatch, Warning, TEXT("ULandscapeTextureHeightPatch::ResizeIfNeeded: Inputs must be positive."));
+		return false;
+	}
+
+	if (SourceMode == ELandscapeTexturePatchSourceMode::TextureBackedRenderTarget)
+	{
+		return ResizeRenderTargetIfNeeded(ResolutionIn.X, ResolutionIn.Y);
+	}
+	else //if (SourceMode == ELandscapeTexturePatchSourceMode::InternalTexture)
+	{
+		return ResizeTextureIfNeeded(ResolutionIn.X, ResolutionIn.Y, false, true);
+	}
+	return false;
+}
+#endif // WITH_EDITOR
+
+
+#undef LOCTEXT_NAMESPACE
\ No newline at end of file
diff --git a/Engine/Plugins/Experimental/LandscapePatch/Source/LandscapePatch/Private/LandscapeTextureHeightPatchPS.cpp b/Engine/Plugins/Experimental/LandscapePatch/Source/LandscapePatch/Private/LandscapeTextureHeightPatchPS.cpp
new file mode 100644
--- /dev/null
+++ b/Engine/Plugins/Experimental/LandscapePatch/Source/LandscapePatch/Private/LandscapeTextureHeightPatchPS.cpp
@@ -0,0 +1,98 @@
+// Copyright Epic Games, Inc. All Rights Reserved.
+
+#include "LandscapeTextureHeightPatchPS.h"
+
+#include "LandscapeUtils.h"
+#include "PixelShaderUtils.h"
+
+namespace UE::Landscape
+{
+
+bool FApplyLandscapeTextureHeightPatchPS::ShouldCompilePermutation(const FGlobalShaderPermutationParameters& Parameters)
+{
+	return DoesPlatformSupportEditLayers(Parameters.Platform);
+}
+
+void FApplyLandscapeTextureHeightPatchPS::ModifyCompilationEnvironment(const FGlobalShaderPermutationParameters& Parameters, FShaderCompilerEnvironment& OutEnvironment)
+{
+	OutEnvironment.SetDefine(TEXT("APPLY_LANDSCAPE_PATCH"), 1);
+}
+
+void FApplyLandscapeTextureHeightPatchPS::AddToRenderGraph(FRDGBuilder& GraphBuilder, FParameters* InParameters, const FIntRect& DestinationBounds)
+{
+	FGlobalShaderMap* ShaderMap = GetGlobalShaderMap(GMaxRHIFeatureLevel);
+	TShaderMapRef<FApplyLandscapeTextureHeightPatchPS> PixelShader(ShaderMap);
+
+	FPixelShaderUtils::AddFullscreenPass(
+		GraphBuilder,
+		ShaderMap,
+		RDG_EVENT_NAME("LandscapeTextureHeightPatch"),
+		PixelShader,
+		InParameters,
+		DestinationBounds);
+}
+
+bool FOffsetHeightmapPS::ShouldCompilePermutation(const FGlobalShaderPermutationParameters& Parameters)
+{
+	return DoesPlatformSupportEditLayers(Parameters.Platform);
+}
+
+void FOffsetHeightmapPS::ModifyCompilationEnvironment(const FGlobalShaderPermutationParameters& Parameters, FShaderCompilerEnvironment& OutEnvironment)
+{
+	OutEnvironment.SetDefine(TEXT("OFFSET_LANDSCAPE_PATCH"), 1);
+}
+
+void FOffsetHeightmapPS::AddToRenderGraph(FRDGBuilder& GraphBuilder, FParameters* InParameters)
+{
+	FGlobalShaderMap* ShaderMap = GetGlobalShaderMap(GMaxRHIFeatureLevel);
+	TShaderMapRef<FOffsetHeightmapPS> PixelShader(ShaderMap);
+
+	FIntVector TextureSize = InParameters->InHeightmap->Desc.Texture->Desc.GetSize();
+
+	FPixelShaderUtils::AddFullscreenPass(
+		GraphBuilder,
+		ShaderMap,
+		RDG_EVENT_NAME("OffsetHeightmap"),
+		PixelShader,
+		InParameters,
+		FIntRect(0, 0, TextureSize.X, TextureSize.Y));
+}
+
+bool FSimpleTextureCopyPS::ShouldCompilePermutation(const FGlobalShaderPermutationParameters& Parameters)
+{
+	return DoesPlatformSupportEditLayers(Parameters.Platform);
+}
+
+void FSimpleTextureCopyPS::ModifyCompilationEnvironment(const FGlobalShaderPermutationParameters& Parameters, FShaderCompilerEnvironment& OutEnvironment)
+{
+	OutEnvironment.SetDefine(TEXT("SIMPLE_TEXTURE_COPY"), 1);
+}
+
+void FSimpleTextureCopyPS::AddToRenderGraph(FRDGBuilder& GraphBuilder, FRDGTextureRef SourceTexture, FRDGTextureRef DestinationTexture)
+{
+	FGlobalShaderMap* ShaderMap = GetGlobalShaderMap(GMaxRHIFeatureLevel);
+	TShaderMapRef<FSimpleTextureCopyPS> PixelShader(ShaderMap);
+
+	FRDGTextureSRVRef SourceSRV = GraphBuilder.CreateSRV(FRDGTextureSRVDesc::CreateForMipLevel(SourceTexture, 0));
+
+	FSimpleTextureCopyPS::FParameters* ShaderParams = GraphBuilder.AllocParameters<FSimpleTextureCopyPS::FParameters>();
+	ShaderParams->InSource = GraphBuilder.CreateSRV(FRDGTextureSRVDesc::CreateForMipLevel(SourceTexture, 0));
+	ShaderParams->InSourceSampler = TStaticSamplerState<SF_Bilinear, AM_Clamp, AM_Clamp>::GetRHI();
+	ShaderParams->RenderTargets[0] = FRenderTargetBinding(DestinationTexture, ERenderTargetLoadAction::ENoAction, /*InMipIndex = */0);
+	FIntVector DestinationSize = DestinationTexture->Desc.GetSize();
+	ShaderParams->InDestinationResolution = FVector2f(DestinationSize.X, DestinationSize.Y);
+
+	FPixelShaderUtils::AddFullscreenPass(
+		GraphBuilder,
+		ShaderMap,
+		RDG_EVENT_NAME("SimpleTextureCopy"),
+		PixelShader,
+		ShaderParams,
+		FIntRect(0, 0, DestinationSize.X, DestinationSize.Y));
+}
+
+}//end UE::Landscape
+
+IMPLEMENT_GLOBAL_SHADER(UE::Landscape::FApplyLandscapeTextureHeightPatchPS, "/Plugin/LandscapePatch/Private/LandscapeTextureHeightPatchPS.usf", "ApplyLandscapeTextureHeightPatch", SF_Pixel);
+IMPLEMENT_GLOBAL_SHADER(UE::Landscape::FOffsetHeightmapPS, "/Plugin/LandscapePatch/Private/LandscapeTextureHeightPatchPS.usf", "ApplyOffsetToHeightmap", SF_Pixel);
+IMPLEMENT_GLOBAL_SHADER(UE::Landscape::FSimpleTextureCopyPS, "/Plugin/LandscapePatch/Private/LandscapeTextureHeightPatchPS.usf", "SimpleTextureCopy", SF_Pixel);
diff --git a/Engine/Plugins/Experimental/LandscapePatch/Source/LandscapePatch/Private/LandscapeTexturePatchBase.cpp b/Engine/Plugins/Experimental/LandscapePatch/Source/LandscapePatch/Private/LandscapeTexturePatchBase.cpp
new file mode 100644
--- /dev/null
+++ b/Engine/Plugins/Experimental/LandscapePatch/Source/LandscapePatch/Private/LandscapeTexturePatchBase.cpp
@@ -0,0 +1,125 @@
+// Copyright Epic Games, Inc. All Rights Reserved.
+
+#include "LandscapeTexturePatchBase.h"
+
+#include "Engine/Texture.h"
+#include "Engine/Texture2D.h"
+#include "Engine/TextureRenderTarget2D.h"
+#include "Landscape.h"
+
+FTransform ULandscapeTexturePatchBase::GetPatchToWorldTransform() const
+{
+	FTransform PatchToWorld = GetComponentTransform();
+
+	if (Landscape.IsValid())
+	{
+		FRotator3d PatchRotator = PatchToWorld.GetRotation().Rotator();
+		FRotator3d LandscapeRotator = Landscape->GetTransform().GetRotation().Rotator();
+		PatchToWorld.SetRotation(FRotator3d(LandscapeRotator.Pitch, PatchRotator.Yaw, LandscapeRotator.Roll).Quaternion());
+	}
+
+	return PatchToWorld;
+}
+
+bool ULandscapeTexturePatchBase::GetTextureResolution(FVector2D& SizeOut) const
+{
+	switch (SourceMode)
+	{
+	case ELandscapeTexturePatchSourceMode::InternalTexture:
+		if (IsValid(InternalTexture))
+		{
+			// Apparently direct GetSizeX/Y calls can return a default texture size in some cases
+			// while the texture is compiling, hence us going through the resource here.
+			FTextureResource* TextureResource = InternalTexture->GetResource();
+			if (ensure(TextureResource))
+			{
+				SizeOut = FVector2D(TextureResource->GetSizeX(), TextureResource->GetSizeY());
+			}
+			else
+			{
+				SizeOut = FVector2D(InternalTexture->GetSizeX(), InternalTexture->GetSizeY());
+			}
+			return true;
+		}
+		break;
+	case ELandscapeTexturePatchSourceMode::TextureBackedRenderTarget:
+		if (IsValid(InternalRenderTarget))
+		{
+			SizeOut = FVector2D(InternalRenderTarget->SizeX, InternalRenderTarget->SizeY);
+			return true;
+		}
+		break;
+	case ELandscapeTexturePatchSourceMode::TextureAsset:
+		if (IsValid(TextureAsset))
+		{
+			FTextureResource* TextureResource = TextureAsset->GetResource();
+			if (ensure(TextureResource))
+			{
+				SizeOut = FVector2D(TextureResource->GetSizeX(), TextureResource->GetSizeY());
+			}
+			else
+			{
+				SizeOut = FVector2D(TextureAsset->GetResource()->GetSizeX(), TextureAsset->GetResource()->GetSizeY());
+			}
+			return true;
+		}
+		break;
+	default:
+		ensure(false);
+		break;
+	}
+
+	return false;
+}
+
+FVector2D ULandscapeTexturePatchBase::GetFullUnscaledWorldSize() const
+{
+	FVector2D Resolution;
+
+	if (!GetTextureResolution(Resolution))
+	{
+		return UnscaledPatchCoverage;
+	}
+
+	return GetFullUnscaledWorldSizeForResolution(Resolution);
+}
+
+FVector2D ULandscapeTexturePatchBase::GetFullUnscaledWorldSizeForResolution(const FVector2D& ResolutionIn) const
+{
+	// UnscaledPatchCoverage is meant to represent the distance between the centers of the extremal pixels.
+	// That distance in pixels is TextureSize-1.
+	FVector2D TargetPixelSize(UnscaledPatchCoverage / FMath::Max(ResolutionIn - 1, FVector2D(1, 1)));
+	return TargetPixelSize * ResolutionIn;
+}
+
+bool ULandscapeTexturePatchBase::GetInitResolutionFromLandscape(float ResolutionMultiplier, FVector2D& ResolutionOut) const
+{
+	if (!Landscape.IsValid())
+	{
+		return false;
+	}
+
+	ResolutionOut = FVector2D::Zero();
+
+	FVector LandscapeScale = Landscape->GetTransform().GetScale3D();
+	// We go off of the larger dimension so that our patch works in different rotations.
+	double LandscapeQuadSize = FMath::Max(FMath::Abs(LandscapeScale.X), FMath::Abs(LandscapeScale.Y));
+
+	if (LandscapeQuadSize > 0)
+	{
+		double PatchQuadSize = LandscapeQuadSize;
+		PatchQuadSize /= (ResolutionMultiplier > 0 ? ResolutionMultiplier : 1);
+
+		FVector PatchScale = GetComponentTransform().GetScale3D();
+		double NumQuadsX = FMath::Abs(UnscaledPatchCoverage.X * PatchScale.X / PatchQuadSize);
+		double NumQuadsY = FMath::Abs(UnscaledPatchCoverage.Y * PatchScale.Y / PatchQuadSize);
+
+		ResolutionOut = FVector2D(
+			FMath::Max(1, FMath::CeilToInt(NumQuadsX) + 1),
+			FMath::Max(1, FMath::CeilToInt(NumQuadsY) + 1)
+		);
+
+		return true;
+	}
+	return false;
+}
diff --git a/Engine/Plugins/Experimental/LandscapePatch/Source/LandscapePatch/Public/LandscapeCircleHeightPatch.h b/Engine/Plugins/Experimental/LandscapePatch/Source/LandscapePatch/Public/LandscapeCircleHeightPatch.h
--- a/Engine/Plugins/Experimental/LandscapePatch/Source/LandscapePatch/Public/LandscapeCircleHeightPatch.h
+++ b/Engine/Plugins/Experimental/LandscapePatch/Source/LandscapePatch/Public/LandscapeCircleHeightPatch.h
@@ -4,7 +4,7 @@
 
 #include "CoreMinimal.h"
 
-#include "LandscapePatch.h"
+#include "LandscapePatchComponent.h"
 
 #include "LandscapeCircleHeightPatch.generated.h"
 
diff --git a/Engine/Plugins/Experimental/LandscapePatch/Source/LandscapePatch/Public/LandscapePatch.h b/Engine/Plugins/Experimental/LandscapePatch/Source/LandscapePatch/Public/LandscapePatch.h
deleted file mode 100644
diff --git a/Engine/Plugins/Experimental/LandscapePatch/Source/LandscapePatch/Public/LandscapePatchComponent.h b/Engine/Plugins/Experimental/LandscapePatch/Source/LandscapePatch/Public/LandscapePatchComponent.h
new file mode 100644
--- /dev/null
+++ b/Engine/Plugins/Experimental/LandscapePatch/Source/LandscapePatch/Public/LandscapePatchComponent.h
@@ -0,0 +1,81 @@
+// Copyright Epic Games, Inc. All Rights Reserved.
+
+#pragma once
+
+#include "CoreMinimal.h"
+
+#include "Components/SceneComponent.h"
+
+#include "LandscapePatchComponent.generated.h"
+
+class ALandscape;
+class ALandscapePatchManager;
+class UTextureRenderTarget2D;
+
+/**
+ * Base class for landscape patches: components that can be attached to meshes and moved around to make
+ * the meshes affect the landscape around themselves.
+ */
+//~ TODO: Although this doesn't generate geometry, we are likely to change this to inherit from UPrimitiveComponent
+//~ so that we can use render proxies for passing along data to the render thread or perhaps for visualization.
+UCLASS(Blueprintable, BlueprintType, Abstract)
+class LANDSCAPEPATCH_API ULandscapePatchComponent : public USceneComponent
+{
+	GENERATED_BODY()
+
+public:
+	ULandscapePatchComponent(const FObjectInitializer& ObjectInitializer = FObjectInitializer::Get());
+
+	virtual void Initialize_Native(const FTransform& InLandscapeTransform,
+		const FIntPoint& InLandscapeSize,
+		const FIntPoint& InLandscapeRenderTargetSize) {}
+	virtual UTextureRenderTarget2D* Render_Native(bool InIsHeightmap,
+		UTextureRenderTarget2D* InCombinedResult,
+		const FName& InWeightmapLayerName) { return InCombinedResult; }
+
+// For now we keep the patches largely editor-only, since we don't yet support runtime landscape editing
+#if WITH_EDITOR
+
+	UFUNCTION(BlueprintCallable, Category = "LandscapePatch")
+	void RequestLandscapeUpdate();
+
+	UFUNCTION(BlueprintCallable, Category = "LandscapePatch")
+	FTransform GetLandscapeHeightmapCoordsToWorld() const;
+
+	// USceneComponent
+	virtual void OnUpdateTransform(EUpdateTransformFlags UpdateTransformFlags, ETeleportType Teleport) override;
+
+	// UActorComponent
+	virtual void OnComponentCreated() override;
+	virtual void OnComponentDestroyed(bool bDestroyingHierarchy) override;
+	virtual void OnRegister() override;
+
+	// UObject
+	virtual void PostEditChangeProperty(FPropertyChangedEvent& PropertyChangedEvent) override;
+	virtual bool IsPostLoadThreadSafe() const override { return true; }
+	virtual void PostLoad();
+#endif
+	virtual bool IsEditorOnly() const override { return true; }
+	virtual bool NeedsLoadForClient() const override { return false; }
+	virtual bool NeedsLoadForServer() const override { return false; }
+
+protected:
+	UPROPERTY(EditAnywhere, Category = Settings)
+	TWeakObjectPtr<ALandscape> Landscape = nullptr;
+
+	UPROPERTY()
+	TWeakObjectPtr<ALandscapePatchManager> PatchManager = nullptr;
+
+	// Determines whether the height patch was made by copying a different height patch.
+	bool bWasCopy = false;
+
+	// This is true for existing height patches right after they are loaded, so that we can ignore
+	// the first OnRegister call. It remains false from the first OnRegiter call onward, even
+	// if the component is unregistered.
+	bool bLoadedButNotYetRegistered = false;
+private:
+	// Starts as false and gets set to true in construction, so gets used to set bWasCopy
+	// by checking the indicator value at the start of construction.
+	UPROPERTY()
+	bool bPropertiesCopiedIndicator = false;
+};
\ No newline at end of file
diff --git a/Engine/Plugins/Experimental/LandscapePatch/Source/LandscapePatch/Public/LandscapePatchManager.h b/Engine/Plugins/Experimental/LandscapePatch/Source/LandscapePatch/Public/LandscapePatchManager.h
--- a/Engine/Plugins/Experimental/LandscapePatch/Source/LandscapePatch/Public/LandscapePatchManager.h
+++ b/Engine/Plugins/Experimental/LandscapePatch/Source/LandscapePatch/Public/LandscapePatchManager.h
@@ -64,6 +64,7 @@ protected:
 	UPROPERTY()
 	TArray<TWeakObjectPtr<ULandscapePatchComponent>> PatchComponents;
 
+	UPROPERTY()
 	FTransform HeightmapCoordsToWorld;
 
 };
\ No newline at end of file
diff --git a/Engine/Plugins/Experimental/LandscapePatch/Source/LandscapePatch/Public/LandscapeTextureHeightPatch.h b/Engine/Plugins/Experimental/LandscapePatch/Source/LandscapePatch/Public/LandscapeTextureHeightPatch.h
new file mode 100644
--- /dev/null
+++ b/Engine/Plugins/Experimental/LandscapePatch/Source/LandscapePatch/Public/LandscapeTextureHeightPatch.h
@@ -0,0 +1,228 @@
+// Copyright Epic Games, Inc. All Rights Reserved.
+
+#pragma once
+
+#include "CoreMinimal.h"
+
+#include "Engine/TextureRenderTarget2D.h"
+
+#include "LandscapeTextureHeightPatchPS.h" // FApplyLandscapeTextureHeightPatchPS::FParameters
+#include "LandscapeTexturePatchBase.h"
+#include "MatrixTypes.h"
+
+#include "LandscapeTextureHeightPatch.generated.h"
+
+class UTexture2D;
+
+// Used for Reinitialize, to determine how to initialize the texture.
+UENUM()
+enum class ELandscapeTextureHeightPatchInitMode : uint8
+{
+	// Sample the current landscape covered by the patch to initialize.
+	FromLandscape,
+	
+	// Initialize from currently set texture asset.
+	TextureAsset,
+
+	// Initialize to landscape mid value.
+	Blank,
+};
+
+// Determines how the patch is combined with the previous state of the landscape.
+UENUM(BlueprintType)
+enum class ELandscapeTextureHeightPatchBlendMode : uint8
+{
+	// Let the patch specify the actual target height, and blend that with the existing
+	// height using falloff/alpha. E.g. with no falloff and alpha 1, the landscape will
+	// be set directly to the height sampled from patch. With alpha 0.5, landscape height 
+	// will be averaged evenly with patch height.
+	AlphaBlend,
+	
+	// Interpreting the landscape mid value as 0, use the texture patch as an offset to
+	// apply to the landscape. Falloff/alpha will just affect the degree to which the offset
+	// is applied (e.g. alpha of 0.5 will apply just half the offset).
+	Additive,
+};
+
+// Determines falloff method for the patch's influence.
+UENUM(BlueprintType)
+enum class ELandscapeTextureHeightPatchFalloffMode : uint8
+{
+	// Affect landscape in a circle inscribed in the patch, and fall off across
+	// a margin extending into that circle.
+	Circle,
+
+	// Affect entire rectangle of patch (except for circular corners), and fall off
+	// across a margin extending inward from the boundary.
+	RoundedRectangle,
+};
+
+/**
+ * A texture-based height patch
+ */
+UCLASS(Blueprintable, BlueprintType, ClassGroup = Landscape, meta = (BlueprintSpawnableComponent))
+class LANDSCAPEPATCH_API ULandscapeTextureHeightPatch : public ULandscapeTexturePatchBase
+{
+	GENERATED_BODY()
+
+public:
+
+	virtual UTextureRenderTarget2D* Render_Native(bool InIsHeightmap,
+		UTextureRenderTarget2D* InCombinedResult,
+		const FName& InWeightmapLayerName) override;
+	
+	// ULandscapeTexturePatchBase
+	virtual bool SetTextureResolution(FVector2D ResolutionIn) override;
+	virtual bool SetSourceMode(ELandscapeTexturePatchSourceMode NewMode, bool bDeleteUnusedInternalTextures = true) override;
+
+	// For now the patch is largely editor-only
+#if WITH_EDITOR
+	// UActorComponent
+	virtual void OnComponentCreated() override;
+
+	// UObject
+	virtual void PreSave(FObjectPreSaveContext SaveContext) override;
+	virtual void PostLoad() override;
+	virtual void PreDuplicate(FObjectDuplicationParameters& DupParams) override;
+	virtual void ExportCustomProperties(FOutputDevice& Out, uint32 Indent) override;
+	virtual void PostEditChangeProperty(FPropertyChangedEvent& PropertyChangedEvent) override;
+	virtual void PostInitProperties() override;
+
+	/**
+	 * Deletes the internal render target and internal texture.
+	 */
+	UFUNCTION(BlueprintCallable, CallInEditor, Category = Initialization)
+	void DeleteInternalTextures();
+
+	/**
+	 * Given the current initialization settings, reinitialize the height patch.
+	 */
+	UFUNCTION(CallInEditor, Category = Initialization)
+	void Reinitialize();
+
+	/**
+	 * Adjusts patch rotation to be aligned to a 90 degree increment relative to the landscape,
+	 * adjusts UnscaledPatchCoverage such that it becomes a multiple of landscape quad size, and
+	 * adjusts patch location so that the boundaries of the covered area lie on the nearest
+	 * landscape vertices.
+	 * Note that this doesn't adjust the resolution of the texture that the patch uses, so landscape
+	 * vertices within the inside of the patch may still not always align with texture patch pixel
+	 * centers (if the resolutions aren't multiples of each other).
+	 */
+	UFUNCTION(BlueprintCallable, CallInEditor, Category = Initialization)
+	void SnapToLandscape();
+#endif
+
+	UFUNCTION(BlueprintCallable, Category = "LandscapePatch")
+	void SetFalloff(float FalloffIn) { Falloff = FalloffIn; }
+
+	UFUNCTION(BlueprintCallable, Category = "LandscapePatch")
+	void SetBlendMode(ELandscapeTextureHeightPatchBlendMode BlendModeIn) { BlendMode = BlendModeIn; }
+
+	/**
+	 * Determine whether the patch Z relative to the landscape affects the results. See comment for bUsePatchZAsReference.
+	 */
+	UFUNCTION(BlueprintCallable, Category = "LandscapePatch")
+	void SetUsePatchZAsReference(bool bUsePatchZIn) { bUsePatchZAsReference = bUsePatchZIn; }
+
+	UFUNCTION(BlueprintCallable, Category = "LandscapePatch")
+	void SetUseTextureAlphaChannel(bool bUse) { bUseTextureAlphaChannel = bUse; }
+
+protected:
+
+	/**
+	 * Extra scaling to apply to patch height. Useful when using adjusting existing texture assets.
+	 */
+	UPROPERTY(EditAnywhere, Category = Settings, meta = (UIMin = "0", UIMax = "1"))
+	float PatchHeightScale = 1;
+
+	/**
+	 * Whether to apply the patch Z scale to the height stored in the patch.
+	 */
+	UPROPERTY(EditAnywhere, Category = Settings, AdvancedDisplay, meta = (DisplayName = "Apply Component Z Scale"))
+	bool bApplyComponentZScale = true;
+
+	/**
+	 * Determine whether the patch Z relative to the landscape affects the results. When true, moving the
+	 * patch up and down will move the results up and down, which is useful for using the patch to move
+	 * landscape edits with a mesh. When false, the patch is always applied as if it is at z=0 relative
+	 * to landscape, which removes the need to worry about its vertical location and is useful when using
+	 * the patch in Blueprints.
+	 * Note that this setting also affects reinitialization from landscape, and it should probably match
+	 * usage when initializing.
+	 */
+	UPROPERTY(EditAnywhere, Category = Settings, AdvancedDisplay)
+	bool bUsePatchZAsReference = true;
+
+	UPROPERTY(EditAnywhere, Category = Settings, AdvancedDisplay)
+	ELandscapeTextureHeightPatchBlendMode BlendMode = ELandscapeTextureHeightPatchBlendMode::AlphaBlend;
+	
+	/** When true, texture alpha channel will be used when applying the patch. */
+	UPROPERTY(EditAnywhere, Category = Settings, AdvancedDisplay)
+	bool bUseTextureAlphaChannel = false;
+
+	UPROPERTY(EditAnywhere, Category = Settings)
+	ELandscapeTextureHeightPatchFalloffMode FalloffMode = ELandscapeTextureHeightPatchFalloffMode::RoundedRectangle;
+
+	/**
+	 * Distance (in unscaled world coordinates) across which to smoothly fall off the patch effects.
+	 */
+	UPROPERTY(EditAnywhere, Category = Settings, meta = (ClampMin = "0", UIMax = "2000"))
+	float Falloff = 0;
+
+	UPROPERTY(EditAnywhere, Category = Initialization, meta = (
+		EditCondition = "SourceMode != ELandscapeTexturePatchSourceMode::TextureAsset", EditConditionHides))
+	ELandscapeTextureHeightPatchInitMode InitializationMode = ELandscapeTextureHeightPatchInitMode::FromLandscape;
+
+	/** When initializing from an texture asset, set the internal texture to have the same resolution. */
+	UPROPERTY(EditAnywhere, Category = Initialization, meta = (
+		EditCondition = "SourceMode != ELandscapeTexturePatchSourceMode::TextureAsset && InitializationMode == ELandscapeTextureHeightPatchInitMode::TextureAsset", 
+		EditConditionHides))
+	bool bUseSameTextureDimensions = false;
+
+	/** When initializing from landscape, set resolution based off of the landscape (and a multiplier). */
+	UPROPERTY(EditAnywhere, Category = Initialization, meta = (
+		EditCondition = "SourceMode != ELandscapeTexturePatchSourceMode::TextureAsset && (InitializationMode != ELandscapeTextureHeightPatchInitMode::TextureAsset || !bUseSameTextureDimensions)", 
+		EditConditionHides))
+	bool bBaseResolutionOffLandscape = true;
+
+	/** 
+	 * Multiplier to apply to landscape resolution when initializing patch resolution. A value greater than 1.0 will use higher
+	 * resolution than the landscape (perhaps useful for slightly more accurate results while not aligned to landscape), and
+	 * a value less that 1.0 will use lower.
+	 */
+	UPROPERTY(EditAnywhere, Category = Initialization, meta = (
+		EditCondition = "SourceMode != ELandscapeTexturePatchSourceMode::TextureAsset && (InitializationMode != ELandscapeTextureHeightPatchInitMode::TextureAsset || !bUseSameTextureDimensions) && bBaseResolutionOffLandscape", 
+		EditConditionHides))
+	float ResolutionMultiplier = 1;
+
+	/** Texture width to use when reinitializing. */
+	UPROPERTY(EditAnywhere, Category = Initialization, meta = (
+		EditCondition = "SourceMode != ELandscapeTexturePatchSourceMode::TextureAsset && (InitializationMode != ELandscapeTextureHeightPatchInitMode::TextureAsset || !bUseSameTextureDimensions)", 
+		EditConditionHides, ClampMin = "1"))
+	int32 InitTextureSizeX = 64;
+
+	/** Texture height to use when reinitializing */
+	UPROPERTY(EditAnywhere, Category = Initialization, meta = (
+		EditCondition = "SourceMode != ELandscapeTexturePatchSourceMode::TextureAsset && (InitializationMode != ELandscapeTextureHeightPatchInitMode::TextureAsset || !bUseSameTextureDimensions)", 
+		EditConditionHides, ClampMin = "1"))
+	int32 InitTextureSizeY = 64;
+
+	// Used to properly transition the source mode when editing it via the detail panel.
+	UPROPERTY()
+	ELandscapeTexturePatchSourceMode PreviousSourceMode = SourceMode;
+
+private:
+
+#if WITH_EDITOR
+	// Uncertain whether these may end up in the base class once we have weight patches.
+	void CopyInternalRenderTargetToInternalTexture(bool bBlock);
+	void CopyInternalTextureToInternalRenderTarget();
+
+	void UpdateShaderParams(UE::Landscape::FApplyLandscapeTextureHeightPatchPS::FParameters& Params, 
+		const FIntPoint& DestinationResolution, FIntRect& DestinationBoundsOut) const;
+
+	bool ResizeRenderTargetIfNeeded(int32 TextureSizeX, int32 TextureSizeY);
+	bool ResizeTextureIfNeeded(int32 TextureSizeX, int32 TextureSizeY, bool bClear, bool bUpdateResource);
+#endif
+};
\ No newline at end of file
diff --git a/Engine/Plugins/Experimental/LandscapePatch/Source/LandscapePatch/Public/LandscapeTextureHeightPatchPS.h b/Engine/Plugins/Experimental/LandscapePatch/Source/LandscapePatch/Public/LandscapeTextureHeightPatchPS.h
new file mode 100644
--- /dev/null
+++ b/Engine/Plugins/Experimental/LandscapePatch/Source/LandscapePatch/Public/LandscapeTextureHeightPatchPS.h
@@ -0,0 +1,97 @@
+// Copyright Epic Games, Inc. All Rights Reserved.
+
+#pragma once
+
+#include "CoreMinimal.h"
+
+#include "GlobalShader.h"
+#include "ShaderParameterStruct.h"
+
+namespace UE::Landscape
+{
+
+/**
+ * Shader that applies a texture-based height patch to a landscape heightmap.
+ */
+class LANDSCAPEPATCH_API FApplyLandscapeTextureHeightPatchPS : public FGlobalShader
+{
+	DECLARE_GLOBAL_SHADER(FApplyLandscapeTextureHeightPatchPS);
+	SHADER_USE_PARAMETER_STRUCT(FApplyLandscapeTextureHeightPatchPS, FGlobalShader);
+
+public:
+
+	BEGIN_SHADER_PARAMETER_STRUCT(FParameters, )
+		SHADER_PARAMETER_RDG_TEXTURE_SRV(Texture2D<float4>, InSourceHeightmap)
+		SHADER_PARAMETER_RDG_TEXTURE_SRV(Texture2D<float4>, InHeightPatch)
+		SHADER_PARAMETER_SAMPLER(SamplerState, InHeightPatchSampler)
+		SHADER_PARAMETER(FMatrix44f, InHeightmapToPatch)
+		// Height value to consider the 0 (deepest) values in height patch
+		SHADER_PARAMETER(float, InPatchBaseHeight)
+		// Amount of the patch edge to not apply in UV space. Generally set to 0.5/Dimensions to avoid applying
+		// the edge half-pixels.
+		SHADER_PARAMETER(FVector2f, InEdgeUVDeadBorder)
+		// In patch texture space, the size of the margin across which the alpha falls from 1 to 0
+		SHADER_PARAMETER(float, InFalloffWorldMargin)
+		SHADER_PARAMETER(FVector2f, InPatchWorldDimensions)
+		SHADER_PARAMETER(float, InHeightScale)
+
+		// TODO: Apparently we're not allowed to use bools, and we're theoretically supposed to pack these into a bitfield.
+		// Will be done in later CL.
+		SHADER_PARAMETER(uint32, bRectangularFalloff)
+		SHADER_PARAMETER(uint32, bApplyPatchAlpha)
+		SHADER_PARAMETER(uint32, bAdditiveMode)
+
+		RENDER_TARGET_BINDING_SLOTS() // Holds our output
+		END_SHADER_PARAMETER_STRUCT()
+
+	static bool ShouldCompilePermutation(const FGlobalShaderPermutationParameters& Parameters);
+	static void ModifyCompilationEnvironment(const FGlobalShaderPermutationParameters& InParameters, FShaderCompilerEnvironment& OutEnvironment);
+	static void AddToRenderGraph(FRDGBuilder& GraphBuilder, FParameters* InParameters, const FIntRect& DestinationBounds);
+};
+
+/**
+ * Simple shader that just offsets each height value in a height patch by a constant.
+ */
+class LANDSCAPEPATCH_API FOffsetHeightmapPS : public FGlobalShader
+{
+	DECLARE_GLOBAL_SHADER(FOffsetHeightmapPS);
+	SHADER_USE_PARAMETER_STRUCT(FOffsetHeightmapPS, FGlobalShader);
+
+public:
+
+	BEGIN_SHADER_PARAMETER_STRUCT(FParameters, )
+		SHADER_PARAMETER_RDG_TEXTURE_SRV(Texture2D<float4>, InHeightmap)
+		SHADER_PARAMETER(float, InHeightOffset)
+		RENDER_TARGET_BINDING_SLOTS() // Holds our output
+		END_SHADER_PARAMETER_STRUCT()
+
+	static bool ShouldCompilePermutation(const FGlobalShaderPermutationParameters& Parameters);
+	static void ModifyCompilationEnvironment(const FGlobalShaderPermutationParameters& InParameters, FShaderCompilerEnvironment& OutEnvironment);
+	static void AddToRenderGraph(FRDGBuilder& GraphBuilder, FParameters* InParameters);
+};
+
+/**
+ * Simple shader for copying textures of potentially different resolutions.
+ */
+//~ Theoretically CopyToResolveTarget or AddCopyToResolveTargetPass should work, but I was unable to make them
+//~ work without lots of complaints from the renderer.
+class LANDSCAPEPATCH_API FSimpleTextureCopyPS : public FGlobalShader
+{
+	DECLARE_GLOBAL_SHADER(FSimpleTextureCopyPS);
+	SHADER_USE_PARAMETER_STRUCT(FSimpleTextureCopyPS, FGlobalShader);
+
+public:
+
+	BEGIN_SHADER_PARAMETER_STRUCT(FParameters, )
+		SHADER_PARAMETER_RDG_TEXTURE_SRV(Texture2D<float4>, InSource)
+		SHADER_PARAMETER_SAMPLER(SamplerState, InSourceSampler)
+		SHADER_PARAMETER(FVector2f, InDestinationResolution)
+		RENDER_TARGET_BINDING_SLOTS() // Holds our output
+		END_SHADER_PARAMETER_STRUCT()
+
+	static bool ShouldCompilePermutation(const FGlobalShaderPermutationParameters& Parameters);
+	static void ModifyCompilationEnvironment(const FGlobalShaderPermutationParameters& InParameters, FShaderCompilerEnvironment& OutEnvironment);
+	static void AddToRenderGraph(FRDGBuilder& GraphBuilder, FRDGTextureRef SourceTexture, FRDGTextureRef DestinationTexture);
+};
+
+}//end UE::Landscape
\ No newline at end of file
diff --git a/Engine/Plugins/Experimental/LandscapePatch/Source/LandscapePatch/Public/LandscapeTexturePatchBase.h b/Engine/Plugins/Experimental/LandscapePatch/Source/LandscapePatch/Public/LandscapeTexturePatchBase.h
new file mode 100644
--- /dev/null
+++ b/Engine/Plugins/Experimental/LandscapePatch/Source/LandscapePatch/Public/LandscapeTexturePatchBase.h
@@ -0,0 +1,165 @@
+// Copyright Epic Games, Inc. All Rights Reserved.
+
+#pragma once
+
+#include "CoreMinimal.h"
+
+#include "LandscapePatchComponent.h"
+
+#include "LandscapeTexturePatchBase.generated.h"
+
+class UTexture;
+class UTexture2D;
+class UTextureRenderTarget2D;
+
+/**
+ * Determines how the patch stores its information, which affects its memory usage in editor (not in runtime,
+ * since patches are baked directly into landscape and removed for runtime).
+ */
+UENUM(BlueprintType)
+enum class ELandscapeTexturePatchSourceMode : uint8
+{
+	/**
+	 * The data will be read from an internally-stored UTexture2D. In this mode, the patch can't be written-to via 
+	 * blueprints, but it avoids storing the extra render target needed for TextureBackedRenderTarget.
+	 */
+	 InternalTexture,
+
+	 /**
+	 * The patch data will be read from an internally-stored render target, which can be written to via Blueprints
+	 * and which gets serialized to an internally stored UTexture2D when needed. Uses double the memory of InternalTexture.
+	 */
+	TextureBackedRenderTarget,
+
+	 /**
+	  * The data will be read from a UTexture asset (which can be a render target). Allows multiple patches
+	  * to share the same texture.
+	  */
+	  TextureAsset
+};
+
+/**
+ * A texture-based landscape patch. Base class for height patches (and, to be implemented, for weight patches).
+ */
+//~ We use the Base suffix here in part to reserve the ULandscapeTexturePatch name for potential later use.
+UCLASS(Blueprintable, BlueprintType, Abstract)
+class LANDSCAPEPATCH_API ULandscapeTexturePatchBase : public ULandscapePatchComponent
+{
+	GENERATED_BODY()
+
+public:
+
+	/**
+	 * Gets an internally-stored render target that can be written to by blueprints.  However, will be null if
+	 * source mode is not set to TextureBackedRenderTarget.
+	 */
+	UFUNCTION(BlueprintCallable, Category = "LandscapePatch")
+	UTextureRenderTarget2D* GetInternalRenderTarget() { return InternalRenderTarget;}
+
+	// These need to be public so that we can take the internal texture and write it to an external one
+	UTexture2D* GetInternalTexture() { return InternalTexture; }
+	void SetTextureAsset(UTexture* TextureIn) { TextureAsset = TextureIn; }
+
+	virtual ELandscapeTexturePatchSourceMode GetSourceMode() const { return SourceMode; }
+	/**
+	 * Changes source mode. When changing between internal texture/rendertarget modes, existing data 
+	 * is copied from one to the other.
+	 * 
+	 * @return true If successful.
+	 */
+	UFUNCTION(BlueprintCallable, Category = "LandscapePatch")
+	virtual bool SetSourceMode(ELandscapeTexturePatchSourceMode NewMode, bool bDiscardUnusedInternalTextures = true) { return false; }
+
+	/**
+	 * Gets the transform from patch to world. The transform is based off of the component
+	 * transform, but with rotation changed to align to the landscape, only using the yaw
+	 * to rotate it relative to the landscape.
+	 */
+	UFUNCTION(BlueprintCallable, Category = "LandscapePatch")
+	virtual FTransform GetPatchToWorldTransform() const;
+
+	/**
+	 * Gives size in unscaled world coordinates (ie before applying patch transform) of the patch as measured 
+	 * between the centers of the outermost pixels. Measuring the coverage this way means that a patch can 
+	 * affect the same region of the landscape regardless of patch resolution.
+	 * This is also the range across which bilinear interpolation always has correct values, so the area outside 
+	 * this center portion is usually set as a "dead" border that doesn't affect the landscape.
+	 */
+	UFUNCTION(BlueprintCallable, Category = "LandscapePatch")
+	virtual FVector2D GetUnscaledCoverage() const { return FVector2D(UnscaledPatchCoverage); }
+
+	/**
+	 * Set the patch coverage (see GetUnscaledCoverage for description).
+	 */
+	UFUNCTION(BlueprintCallable, Category = "LandscapePatch")
+	virtual void SetUnscaledCoverage(FVector2D Coverage) { UnscaledPatchCoverage = Coverage; }
+
+	/**
+	 * Gives size in unscaled world coordinates of the patch in the world, based off of UnscaledCoverage and
+	 * texture resolution (i.e., adds a half-pixel around UnscaledCoverage).
+	 */
+	UFUNCTION(BlueprintCallable, Category = "LandscapePatch")
+	virtual FVector2D GetFullUnscaledWorldSize() const;
+
+	/** 
+	 * Gets the size (in pixels) of the currently used texture. Depends on SourceMode, 
+	 * and returns false if currently used texture is not allocated/set.
+	 */
+	UFUNCTION(BlueprintCallable, Category = "LandscapePatch")
+	virtual UPARAM(DisplayName = "Success") bool GetTextureResolution(FVector2D& SizeOut) const;
+
+	/**
+	 * Sets the resolution of the currently used internal texture or render target. Has no effect
+	 * if the source mode is set to an external texture asset.
+	 * 
+	 * @return true if successful.
+	 */
+	UFUNCTION(BlueprintCallable, Category = "LandscapePatch")
+	virtual UPARAM(DisplayName = "Success") bool SetTextureResolution(FVector2D ResolutionIn) { return false; }
+
+	/**
+	 * Given the landscape resolution, current patch coverage, and a landscape resolution multiplier, gives the
+	 * needed resolution of the landscape patch. I.e., figures out the number of pixels in the landcape that
+	 * would be in a region of such size, and then uses the resolution multiplier to give a result.
+	 * 
+	 * @return true if successful (may fail if landscape is not set, for instance)
+	 */
+	UFUNCTION(BlueprintCallable, Category = "LandscapePatch", meta = (ResolutionMultiplier = "1.0"))
+	virtual UPARAM(DisplayName = "Success") bool GetInitResolutionFromLandscape(float ResolutionMultiplier, FVector2D& ResolutionOut) const;
+
+protected:
+	/**
+	 * How the heightmap of the patch is stored.
+	 */
+	UPROPERTY(EditAnywhere, Category = Settings, meta = (DisplayPriority = 10))
+	ELandscapeTexturePatchSourceMode SourceMode = ELandscapeTexturePatchSourceMode::InternalTexture;
+
+	/**
+	 * Texture used when source mode is set to a texture asset, or for intializing from a texture asset.
+	 */
+	UPROPERTY(EditAnywhere, Category = Settings, meta = (EditConditionHides, DisplayPriority = 20, DisplayName = "Texture Asset", 
+		EditCondition = "SourceMode == ELandscapeTexturePatchSourceMode::TextureAsset || bShowTextureAssetProperty"))
+	TObjectPtr<UTexture> TextureAsset = nullptr;
+
+	/** At scale 1.0, the X and Y of the region affected by the height patch. This corresponds to the distance from the center
+	 of the first pixel to the center of the last pixel in the patch texture in the X and Y directions. */
+	UPROPERTY(EditAnywhere, Category = Settings, meta = (UIMin = "0", ClampMin = "0", DisplayPriority = 30))
+	FVector2D UnscaledPatchCoverage = FVector2D(2000, 2000);
+
+	/** Not directly settable via detail panel- for display/debugging purposes only. */
+	UPROPERTY(VisibleAnywhere, Category = Settings, AdvancedDisplay, meta = (EditConditionHides, DisplayPriority = 40,
+		EditCondition = "SourceMode != ELandscapeTexturePatchSourceMode::TextureAsset"))
+	TObjectPtr<UTexture2D> InternalTexture = nullptr;
+
+	/** Not directly settable via detail panel- for display/debugging purposes only. */
+	UPROPERTY(VisibleAnywhere, Category = Settings, Transient, DuplicateTransient, AdvancedDisplay, meta = (DisplayPriority = 50,
+		EditConditionHides, EditCondition = "SourceMode == ELandscapeTexturePatchSourceMode::TextureBackedRenderTarget"))
+	TObjectPtr<UTextureRenderTarget2D> InternalRenderTarget = nullptr;
+
+	// Lets subclasses control whether the TextureAsset property is visible.
+	UPROPERTY()
+	bool bShowTextureAssetProperty;
+
+	// Adds a half-pixel around the border of UnscaledPatchCoverage, where pixel size depends on given resolution
+	FVector2D GetFullUnscaledWorldSizeForResolution(const FVector2D& ResolutionIn) const;
+};
\ No newline at end of file
diff --git a/Engine/Plugins/Experimental/LandscapePatch/Source/LandscapePatchEditorOnly/LandscapePatchEditorOnly.Build.cs b/Engine/Plugins/Experimental/LandscapePatch/Source/LandscapePatchEditorOnly/LandscapePatchEditorOnly.Build.cs
new file mode 100644
--- /dev/null
+++ b/Engine/Plugins/Experimental/LandscapePatch/Source/LandscapePatchEditorOnly/LandscapePatchEditorOnly.Build.cs
@@ -0,0 +1,45 @@
+// Copyright Epic Games, Inc. All Rights Reserved.
+
+using UnrealBuildTool;
+
+public class LandscapePatchEditorOnly : ModuleRules
+{
+	public LandscapePatchEditorOnly(ReadOnlyTargetRules Target) : base(Target)
+	{
+		PCHUsage = ModuleRules.PCHUsageMode.UseExplicitOrSharedPCHs;
+		
+		PublicIncludePaths.AddRange(
+			new string[] {
+				// ... add public include paths required here ...
+			}
+			);
+				
+		
+		PrivateIncludePaths.AddRange(
+			new string[] {
+				// ... add other private include paths required here ...
+			}
+			);
+			
+		
+		PublicDependencyModuleNames.AddRange(
+			new string[]
+			{
+				"Core",
+				"CoreUObject",
+			}
+			);
+			
+		
+		PrivateDependencyModuleNames.AddRange(
+			new string[]
+			{
+				"Engine",
+				"LandscapePatch",
+				"Slate",
+				"SlateCore",
+				"UnrealEd",
+			}
+			);
+	}
+}
diff --git a/Engine/Plugins/Experimental/LandscapePatch/Source/LandscapePatchEditorOnly/Private/LandscapePatchEditorOnlyModule.cpp b/Engine/Plugins/Experimental/LandscapePatch/Source/LandscapePatchEditorOnly/Private/LandscapePatchEditorOnlyModule.cpp
new file mode 100644
--- /dev/null
+++ b/Engine/Plugins/Experimental/LandscapePatch/Source/LandscapePatchEditorOnly/Private/LandscapePatchEditorOnlyModule.cpp
@@ -0,0 +1,60 @@
+// Copyright Epic Games, Inc. All Rights Reserved.
+
+#include "LandscapePatchEditorOnlyModule.h"
+
+#include "Editor/UnrealEdEngine.h"
+#include "LandscapeTexturePatchBase.h"
+#include "LandscapeTexturePatchCustomization.h"
+#include "LandscapeTexturePatchVisualizer.h"
+#include "PropertyEditorModule.h"
+#include "UnrealEdGlobals.h"
+
+#define LOCTEXT_NAMESPACE "FLandscapePatchEditorOnlyModule"
+
+void FLandscapePatchEditorOnlyModule::StartupModule()
+{
+	if (GUnrealEd)
+	{
+		TSharedPtr<FLandscapeTexturePatchVisualizer> Visualizer = MakeShared<FLandscapeTexturePatchVisualizer>();
+		GUnrealEd->RegisterComponentVisualizer(ULandscapeTexturePatchBase::StaticClass()->GetFName(), Visualizer);
+		// This call should maybe be inside the RegisterComponentVisualizer call above, but since it's not,
+		// we'll put it here.
+		Visualizer->OnRegister();
+		VisualizersToUnregisterOnShutdown.Add(ULandscapeTexturePatchBase::StaticClass()->GetFName());
+	}
+
+	// Register detail customization
+	ClassesToUnregisterOnShutdown.Reset();
+	FPropertyEditorModule* PropertyEditorModule = FModuleManager::GetModulePtr<FPropertyEditorModule>("PropertyEditor");
+	if (PropertyEditorModule)
+	{
+		PropertyEditorModule->RegisterCustomClassLayout(ULandscapeTexturePatchBase::StaticClass()->GetFName(),
+			FOnGetDetailCustomizationInstance::CreateStatic(&FLandscapeTexturePatchCustomization::MakeInstance));
+		ClassesToUnregisterOnShutdown.Add(ULandscapeTexturePatchBase::StaticClass()->GetFName());
+	}
+}
+
+void FLandscapePatchEditorOnlyModule::ShutdownModule()
+{
+	if (GUnrealEd)
+	{
+		for (const FName& Name : VisualizersToUnregisterOnShutdown)
+		{
+			GUnrealEd->UnregisterComponentVisualizer(Name);
+		}
+	}
+
+	// Unregister customizations
+	FPropertyEditorModule* PropertyEditorModule = FModuleManager::GetModulePtr<FPropertyEditorModule>("PropertyEditor");
+	if (PropertyEditorModule)
+	{
+		for (const FName& ClassName : ClassesToUnregisterOnShutdown)
+		{
+			PropertyEditorModule->UnregisterCustomClassLayout(ClassName);
+		}
+	}
+}
+
+#undef LOCTEXT_NAMESPACE
+	
+IMPLEMENT_MODULE(FLandscapePatchEditorOnlyModule, LandscapePatchEditorOnly)
\ No newline at end of file
diff --git a/Engine/Plugins/Experimental/LandscapePatch/Source/LandscapePatchEditorOnly/Private/LandscapeTexturePatchCustomization.cpp b/Engine/Plugins/Experimental/LandscapePatch/Source/LandscapePatchEditorOnly/Private/LandscapeTexturePatchCustomization.cpp
new file mode 100644
--- /dev/null
+++ b/Engine/Plugins/Experimental/LandscapePatch/Source/LandscapePatchEditorOnly/Private/LandscapeTexturePatchCustomization.cpp
@@ -0,0 +1,108 @@
+// Copyright Epic Games, Inc. All Rights Reserved.
+
+#include "LandscapeTexturePatchCustomization.h"
+
+#include "DetailCategoryBuilder.h"
+#include "DetailLayoutBuilder.h"
+#include "DetailWidgetRow.h"
+#include "FileHelpers.h"
+#include "LandscapeTexturePatchBase.h"
+#include "TextureCompiler.h"
+#include "Widgets/Input/SButton.h"
+#include "Widgets/Layout/SWrapBox.h"
+
+#define LOCTEXT_NAMESPACE "LandscapeTexturePatchCustomization"
+
+TSharedRef<IDetailCustomization> FLandscapeTexturePatchCustomization::MakeInstance()
+{
+	return MakeShareable(new FLandscapeTexturePatchCustomization);
+}
+
+void FLandscapeTexturePatchCustomization::CustomizeDetails(IDetailLayoutBuilder& DetailBuilder)
+{
+	// We're going to add a button to save the internal texture. We do this in a detail customization
+	// rather than in a "CallInEditor" method on the patch itself because this is an editor-only method,
+	// and the patch is not in an editor-only module.
+
+	TArray<TWeakObjectPtr<UObject>> ObjectsBeingCustomized;
+	DetailBuilder.GetObjectsBeingCustomized(ObjectsBeingCustomized);
+	if (ObjectsBeingCustomized.Num() != 1)
+	{
+		// TODO: Could consider showing but disabling if multiple patches are selected.
+		return;
+	}
+
+	ULandscapeTexturePatchBase* PatchObject = Cast<ULandscapeTexturePatchBase>(ObjectsBeingCustomized[0]);
+	if (!ensure(PatchObject))
+	{
+		return;
+	}
+
+	FText ButtonLabel = LOCTEXT("SaveInternalTextureButtonLabel", "Save Internal Texture as Asset");
+
+	DetailBuilder.EditCategory("Initialization")
+		.AddCustomRow(ButtonLabel, false)
+		.RowTag(FName(TEXT("InternalTextureToExternal")))
+		[
+			// The slate here matches what we do in FObjectDetails::AddCallInEditorMethods
+			SNew(SWrapBox)
+			.UseAllottedSize(true)
+			+ SWrapBox::Slot()
+			.Padding(0.0f, 0.0f, 5.0f, 3.0f)
+			[
+				SNew(SButton)
+				.Text(ButtonLabel)
+				.IsEnabled(TAttribute<bool>::CreateLambda([PatchObject]() {
+					return IsValid(PatchObject) && (
+						(PatchObject->GetSourceMode() == ELandscapeTexturePatchSourceMode::InternalTexture && PatchObject->GetInternalTexture())
+						|| (PatchObject->GetSourceMode() == ELandscapeTexturePatchSourceMode::TextureBackedRenderTarget && PatchObject->GetInternalRenderTarget()));
+				}))
+				.OnClicked(FOnClicked::CreateWeakLambda(PatchObject, [PatchObject]() {
+
+					if (!IsValid(PatchObject) || PatchObject->GetSourceMode() == ELandscapeTexturePatchSourceMode::TextureAsset)
+					{
+						return FReply::Unhandled();
+					}
+
+					if (PatchObject->GetSourceMode() == ELandscapeTexturePatchSourceMode::TextureBackedRenderTarget)
+					{
+						// Make sure the data is inside the internal texture for copying...
+						PatchObject->SetSourceMode(ELandscapeTexturePatchSourceMode::InternalTexture);
+					}
+
+					if (!PatchObject->GetInternalTexture())
+					{
+						return FReply::Unhandled();
+					}
+
+					// For the "save as" call to work properly for us, we would like the internal texture
+					// to be marked as transient and be in the transient package. We could do that with a
+					// Rename(), but feels risky to break it under the patch, even though we'll be changing
+					// the source mode. So we'll make a copy and use that.
+					UTexture* TextureCopy = DuplicateObject(PatchObject->GetInternalTexture(), nullptr);
+					TextureCopy->SetFlags(RF_Transient);
+					FTextureCompilingManager::Get().FinishCompilation({ TextureCopy });
+
+					// Bring up the popup and deal with the saving.
+					TArray<UObject*> SavedObjects;
+					FEditorFileUtils::SaveAssetsAs({ TextureCopy }, SavedObjects);
+
+					if (SavedObjects.Num() > 0)
+					{
+						UTexture* NewTexture = Cast<UTexture>(SavedObjects[0]);
+						if (ensure(NewTexture))
+						{
+							PatchObject->SetTextureAsset(NewTexture);
+						}
+						PatchObject->SetSourceMode(ELandscapeTexturePatchSourceMode::TextureAsset);
+					}
+				
+					return FReply::Handled();
+					}))
+				.ToolTipText(LOCTEXT("SaveInternalTextureTooltip", 
+					"Save the current internal texture as a new texture asset and set that as the source of the patch."))
+			]
+		];
+}
+
+#undef LOCTEXT_NAMESPACE
\ No newline at end of file
diff --git a/Engine/Plugins/Experimental/LandscapePatch/Source/LandscapePatchEditorOnly/Private/LandscapeTexturePatchVisualizer.cpp b/Engine/Plugins/Experimental/LandscapePatch/Source/LandscapePatchEditorOnly/Private/LandscapeTexturePatchVisualizer.cpp
new file mode 100644
--- /dev/null
+++ b/Engine/Plugins/Experimental/LandscapePatch/Source/LandscapePatchEditorOnly/Private/LandscapeTexturePatchVisualizer.cpp
@@ -0,0 +1,28 @@
+// Copyright Epic Games, Inc. All Rights Reserved.
+
+#include "LandscapeTexturePatchVisualizer.h"
+
+#include "Landscape.h"
+#include "LandscapeTexturePatchBase.h"
+
+#include "SceneManagement.h" // FPrimitiveDrawInterface
+
+void FLandscapeTexturePatchVisualizer::DrawVisualization(const UActorComponent* Component, const FSceneView* View, FPrimitiveDrawInterface* PDI)
+{
+	const ULandscapeTexturePatchBase* Patch = Cast<ULandscapeTexturePatchBase>(Component);
+	if (!ensure(Patch))
+	{
+		return;
+	}
+	
+	FTransform PatchToWorld = Patch->GetPatchToWorldTransform();
+	
+	FColor Color = FColor::Red;
+	float Thickness = 3;
+	float DepthBias = 1;
+	bool bScreenSpace = 1;
+
+	DrawRectangle(PDI, PatchToWorld.GetTranslation(), PatchToWorld.GetUnitAxis(EAxis::X), PatchToWorld.GetUnitAxis(EAxis::Y),
+		Color, Patch->GetUnscaledCoverage().X * PatchToWorld.GetScale3D().X, Patch->GetUnscaledCoverage().Y * PatchToWorld.GetScale3D().Y, SDPG_Foreground,
+		Thickness, DepthBias, bScreenSpace);
+}
\ No newline at end of file
diff --git a/Engine/Plugins/Experimental/LandscapePatch/Source/LandscapePatchEditorOnly/Public/LandscapePatchEditorOnlyModule.h b/Engine/Plugins/Experimental/LandscapePatch/Source/LandscapePatchEditorOnly/Public/LandscapePatchEditorOnlyModule.h
new file mode 100644
--- /dev/null
+++ b/Engine/Plugins/Experimental/LandscapePatch/Source/LandscapePatchEditorOnly/Public/LandscapePatchEditorOnlyModule.h
@@ -0,0 +1,20 @@
+// Copyright Epic Games, Inc. All Rights Reserved.
+
+#pragma once
+
+#include "CoreMinimal.h"
+#include "Modules/ModuleManager.h"
+
+class FLandscapePatchEditorOnlyModule : public IModuleInterface
+{
+public:
+
+	/** IModuleInterface implementation */
+	virtual void StartupModule() override;
+	virtual void ShutdownModule() override;
+
+private:
+	/** StaticClass is not safe on shutdown, so we cache the name, and use this to unregister on shut down */
+	TArray<FName> ClassesToUnregisterOnShutdown;
+	TArray<FName> VisualizersToUnregisterOnShutdown;
+};
diff --git a/Engine/Plugins/Experimental/LandscapePatch/Source/LandscapePatchEditorOnly/Public/LandscapeTexturePatchCustomization.h b/Engine/Plugins/Experimental/LandscapePatch/Source/LandscapePatchEditorOnly/Public/LandscapeTexturePatchCustomization.h
new file mode 100644
--- /dev/null
+++ b/Engine/Plugins/Experimental/LandscapePatch/Source/LandscapePatchEditorOnly/Public/LandscapeTexturePatchCustomization.h
@@ -0,0 +1,19 @@
+// Copyright Epic Games, Inc. All Rights Reserved.
+
+#pragma once
+
+#include "IDetailCustomization.h"
+#include "Templates/SharedPointer.h"
+
+class IDetailLayoutBuilder;
+
+/**
+ * Adds the editor-only button for saving the internal texture to the detail panel display of 
+ * landscape texture patches.
+ */ 
+class FLandscapeTexturePatchCustomization : public IDetailCustomization
+{
+public:
+	static TSharedRef<IDetailCustomization> MakeInstance();
+	void CustomizeDetails(IDetailLayoutBuilder& DetailBuilder) override;
+};
diff --git a/Engine/Plugins/Experimental/LandscapePatch/Source/LandscapePatchEditorOnly/Public/LandscapeTexturePatchVisualizer.h b/Engine/Plugins/Experimental/LandscapePatch/Source/LandscapePatchEditorOnly/Public/LandscapeTexturePatchVisualizer.h
new file mode 100644
--- /dev/null
+++ b/Engine/Plugins/Experimental/LandscapePatch/Source/LandscapePatchEditorOnly/Public/LandscapeTexturePatchVisualizer.h
@@ -0,0 +1,16 @@
+// Copyright Epic Games, Inc. All Rights Reserved.
+
+#pragma once
+
+#include "CoreMinimal.h"
+
+#include "ComponentVisualizer.h"
+
+/**
+ * Adds a rectangle denoting the area covered by a landscape texture patch when it is selected in editor.
+ */
+class LANDSCAPEPATCHEDITORONLY_API FLandscapeTexturePatchVisualizer : public FComponentVisualizer
+{
+private:
+	virtual void DrawVisualization(const UActorComponent* Component, const FSceneView* View, FPrimitiveDrawInterface* PDI) override;
+};
\ No newline at end of file
diff --git a/Engine/Shaders/Private/Landscape/LandscapeCommon.ush b/Engine/Shaders/Private/Landscape/LandscapeCommon.ush
--- a/Engine/Shaders/Private/Landscape/LandscapeCommon.ush
+++ b/Engine/Shaders/Private/Landscape/LandscapeCommon.ush
@@ -11,6 +11,10 @@ float2 PackHeight(float InHeight)
 
 float UnpackHeight(float2 InPackedHeight)
 {
-	// Convert each of the two bytes back to [0,255], then assemble
-	return float(((int) round(InPackedHeight.r * 255.0) << 8) | (int) round(InPackedHeight.g * 255.0));
+	// We used to do the inverse of PackHeight here, like this:
+	// return float(((int) round(InPackedHeight.r * 255.0) << 8) | (int) round(InPackedHeight.g * 255.0));
+	
+	// However, this isn't great if you just interpolate between packed heights because that truncates the interpolation
+	// of the higher order byte. So instead we do it by multiplication, which ends up simpler.
+	return (InPackedHeight.r * 256 + InPackedHeight.g) * 255;
 }
\ No newline at end of file
diff --git a/Engine/Source/Runtime/Landscape/Public/LandscapeUtils.h b/Engine/Source/Runtime/Landscape/Public/LandscapeUtils.h
--- a/Engine/Source/Runtime/Landscape/Public/LandscapeUtils.h
+++ b/Engine/Source/Runtime/Landscape/Public/LandscapeUtils.h
@@ -12,6 +12,6 @@ namespace UE::Landscape
 * Returns true if edit layers (GPU landscape tools) are enabled on this platform :
 * Note: this is intended for the editor but is in runtime code since global shaders need to exist in runtime modules 
 */
-bool DoesPlatformSupportEditLayers(EShaderPlatform InShaderPlatform);
+LANDSCAPE_API bool DoesPlatformSupportEditLayers(EShaderPlatform InShaderPlatform);
 
 } // end namespace UE::Landscape
