LandscapePatch: Fix NonUnity Compile

diff --git a/Engine/Plugins/Experimental/LandscapePatch/Source/LandscapePatch/Private/LandscapeTextureHeightPatch.cpp b/Engine/Plugins/Experimental/LandscapePatch/Source/LandscapePatch/Private/LandscapeTextureHeightPatch.cpp
--- a/Engine/Plugins/Experimental/LandscapePatch/Source/LandscapePatch/Private/LandscapeTextureHeightPatch.cpp
+++ b/Engine/Plugins/Experimental/LandscapePatch/Source/LandscapePatch/Private/LandscapeTextureHeightPatch.cpp
@@ -8,6 +8,7 @@
 #include "LandscapePatchLogging.h"
 #include "LandscapePatchManager.h"
 #include "RenderGraph.h" // RDG_EVENT_NAME
+#include "RHIStaticStates.h"
 #include "TextureCompiler.h"
 #include "UObject/ObjectSaveContext.h"
 
