Fix non-unity
#rnx

diff --git a/Engine/Plugins/Experimental/LandscapePatch/Source/LandscapePatchEditorOnly/Private/LandscapeTexturePatchCustomization.cpp b/Engine/Plugins/Experimental/LandscapePatch/Source/LandscapePatchEditorOnly/Private/LandscapeTexturePatchCustomization.cpp
--- a/Engine/Plugins/Experimental/LandscapePatch/Source/LandscapePatchEditorOnly/Private/LandscapeTexturePatchCustomization.cpp
+++ b/Engine/Plugins/Experimental/LandscapePatch/Source/LandscapePatchEditorOnly/Private/LandscapeTexturePatchCustomization.cpp
@@ -5,6 +5,7 @@
 #include "DetailCategoryBuilder.h"
 #include "DetailLayoutBuilder.h"
 #include "DetailWidgetRow.h"
+#include "Engine/Texture2D.h"
 #include "FileHelpers.h"
 #include "LandscapeTexturePatchBase.h"
 #include "TextureCompiler.h"
diff --git a/Engine/Plugins/VirtualProduction/RemoteControl/Source/RemoteControlUI/Private/UI/SRCPanelExposedEntitiesList.cpp b/Engine/Plugins/VirtualProduction/RemoteControl/Source/RemoteControlUI/Private/UI/SRCPanelExposedEntitiesList.cpp
--- a/Engine/Plugins/VirtualProduction/RemoteControl/Source/RemoteControlUI/Private/UI/SRCPanelExposedEntitiesList.cpp
+++ b/Engine/Plugins/VirtualProduction/RemoteControl/Source/RemoteControlUI/Private/UI/SRCPanelExposedEntitiesList.cpp
@@ -20,6 +20,8 @@
 #include "SRCPanelFieldGroup.h"
 #include "SRCPanelExposedField.h"
 #include "UObject/Object.h"
+#include "Widgets/Input/SButton.h"
+#include "Widgets/Layout/SSpacer.h"
 #include "Widgets/Views/STableRow.h"
 #include "Widgets/Views/STreeView.h"
 
diff --git a/Engine/Source/Runtime/Core/Private/ProfilingDebugging/MemoryTrace.cpp b/Engine/Source/Runtime/Core/Private/ProfilingDebugging/MemoryTrace.cpp
--- a/Engine/Source/Runtime/Core/Private/ProfilingDebugging/MemoryTrace.cpp
+++ b/Engine/Source/Runtime/Core/Private/ProfilingDebugging/MemoryTrace.cpp
@@ -5,6 +5,9 @@
 #include "ProfilingDebugging/CallstackTrace.h"
 #include "ProfilingDebugging/TagTrace.h"
 #include "ProfilingDebugging/TraceMalloc.h"
+#include "Trace/Detail/EventNode.h"
+#include "Trace/Detail/LogScope.inl"
+#include "Trace/Detail/Important/ImportantLogScope.inl"
 #include "Trace/Trace.h"
 
 #if UE_TRACE_ENABLED
