Fix non-editor build of LandscapePatch plugin.

diff --git a/Engine/Plugins/Experimental/LandscapePatch/Source/LandscapePatch/Private/LandscapePatchComponent.cpp b/Engine/Plugins/Experimental/LandscapePatch/Source/LandscapePatch/Private/LandscapePatchComponent.cpp
--- a/Engine/Plugins/Experimental/LandscapePatch/Source/LandscapePatch/Private/LandscapePatchComponent.cpp
+++ b/Engine/Plugins/Experimental/LandscapePatch/Source/LandscapePatch/Private/LandscapePatchComponent.cpp
@@ -36,7 +36,9 @@ namespace LandscapePatchComponentLocals
 		return PatchManager;
 	}
 }
+#endif
 
+// Note that this is not allowed to be editor-only
 ULandscapePatchComponent::ULandscapePatchComponent(const FObjectInitializer& ObjectInitializer)
 	: Super(ObjectInitializer)
 {
@@ -47,6 +49,7 @@ ULandscapePatchComponent::ULandscapePatchComponent(const FObjectInitializer& Obj
 	bWantsOnUpdateTransform = true;
 }
 
+#if WITH_EDITOR
 void ULandscapePatchComponent::OnComponentCreated()
 {
 	using namespace LandscapePatchComponentLocals;
diff --git a/Engine/Plugins/Experimental/LandscapePatch/Source/LandscapePatch/Private/LandscapeTextureHeightPatch.cpp b/Engine/Plugins/Experimental/LandscapePatch/Source/LandscapePatch/Private/LandscapeTextureHeightPatch.cpp
--- a/Engine/Plugins/Experimental/LandscapePatch/Source/LandscapePatch/Private/LandscapeTextureHeightPatch.cpp
+++ b/Engine/Plugins/Experimental/LandscapePatch/Source/LandscapePatch/Private/LandscapeTextureHeightPatch.cpp
@@ -151,7 +151,6 @@ void ULandscapeTextureHeightPatch::ExportCustomProperties(FOutputDevice& Out, ui
 	}
 }
 
-#if WITH_EDITOR
 void ULandscapeTextureHeightPatch::PostEditChangeProperty(FPropertyChangedEvent& PropertyChangedEvent)
 {
 	if (PropertyChangedEvent.Property)
@@ -175,7 +174,6 @@ void ULandscapeTextureHeightPatch::PostInitProperties()
 
 	Super::PostInitProperties();
 }
-#endif
 
 void ULandscapeTextureHeightPatch::OnComponentCreated()
 {
