LandscapePatch: Add ability to change height encoding mode in texture patches to make it easier to use texture assets, or to use the internal render target. Convert to/from native encoding when needing to save/load. Disallow virtual textures to avoid hitting checks. Pack bool shader parameters into a bitfield.

diff --git a/Engine/Plugins/Experimental/LandscapePatch/Shaders/Private/LandscapeTextureHeightPatchPS.usf b/Engine/Plugins/Experimental/LandscapePatch/Shaders/Private/LandscapeTextureHeightPatchPS.usf
--- a/Engine/Plugins/Experimental/LandscapePatch/Shaders/Private/LandscapeTextureHeightPatchPS.usf
+++ b/Engine/Plugins/Experimental/LandscapePatch/Shaders/Private/LandscapeTextureHeightPatchPS.usf
@@ -10,32 +10,43 @@
 //#define APPLY_LANDSCAPE_PATCH 1
 //#define OFFSET_LANDSCAPE_PATCH 1
 //#define SIMPLE_TEXTURE_COPY 1
+//#define CONVERT_TO_NATIVE_LANDSCAPE_PATCH 1
+//#define CONVERT_BACK_FROM_NATIVE_LANDSCAPE_PATCH 1
 #endif // defined (__INTELLISENSE__)
 
+// This constant gets used in multiple shaders, hence it's out here.
+static const float LANDSCAPE_MID_VALUE = 32768.0f;
+
 #if APPLY_LANDSCAPE_PATCH
 Texture2D<float4> InSourceHeightmap;
 Texture2D<float4> InHeightPatch;
 SamplerState InHeightPatchSampler;
 float4x4 InHeightmapToPatch;
 float2 InPatchWorldDimensions;
-float InPatchBaseHeight;
 float2 InEdgeUVDeadBorder;
 float InFalloffWorldMargin;
+// Value that corresponds to 0 (i.e. to LANDSCAPE_MID_VALUE) in the patch data
+float InZeroInEncoding;
+// Scale to apply to patch values relative to the value representing zero
 float InHeightScale;
-// When false, falloff is circular. When true, it is rectangular.
-uint bRectangularFalloff;
-// When true, the texture alpha channel is considered for blending (in addition to falloff, if any)
-uint bApplyPatchAlpha;
-// When true, the patch stored height is interpreted relative to landscape mid value and applied to
-// the landscape, instead of alpha blending with the landscape.
-uint bAdditiveMode;
+// An offset to apply to the resulting height, after applying the height scale
+float InHeightOffset;
+// A combination of flags, whose positions are set in the corresponding cpp file.
+uint InFlags;
+
 
 void ApplyLandscapeTextureHeightPatch(in float4 SVPos : SV_POSITION, out float2 OutColor : SV_Target0)
 {
 	static const float HALF_PI = 3.14159265f / 2;
-	static const float LANDSCAPE_MID_VALUE = 32768.0f;
 	static const float FLOAT_TOLERANCE = 0.0001;
 	
+	// Unpack our flags. The flag choices are set in FApplyLandscapeTextureHeightPatchPS::ModifyCompilationEnvironment
+	// so that they match to the cpp file
+	bool bRectangularFalloff = InFlags & RECTANGULAR_FALLOFF_FLAG;
+	bool bApplyPatchAlpha = InFlags & APPLY_PATCH_ALPHA_FLAG;
+	bool bAdditiveMode = InFlags & ADDITIVE_MODE_FLAG;
+	bool bInputIsPackedHeight = InFlags & INPUT_IS_PACKED_HEIGHT_FLAG;
+	
 	// We need only the 2D affine transformation that goes from landscape XY heightmap integer coordinates
 	// to patch UV coordinates.
 	float2x2 HeightmapToPatchRotateScale = (float2x2) InHeightmapToPatch;
@@ -43,12 +54,9 @@ void ApplyLandscapeTextureHeightPatch(in float4 SVPos : SV_POSITION, out float2
 	
 	float2 PatchUVCoordinates = mul(HeightmapToPatchRotateScale, SVPos.xy) + HeightmapToPatchTranslate;
 	float4 PatchSampledValue = InHeightPatch.Sample(InHeightPatchSampler, PatchUVCoordinates);
-	float PatchStoredHeight = UnpackHeight(PatchSampledValue.xy);
-	if (bAdditiveMode)
-	{
-		PatchStoredHeight -= LANDSCAPE_MID_VALUE;
-	}
-	float PatchHeight = InHeightScale * PatchStoredHeight + InPatchBaseHeight;
+	float PatchStoredHeight = bInputIsPackedHeight ? UnpackHeight(PatchSampledValue.xy) : PatchSampledValue.x;
+	
+	float PatchSignedHeight = InHeightScale * (PatchStoredHeight - InZeroInEncoding) + InHeightOffset;
 	
 	int2 HeightmapCoordinates = floor(SVPos.xy);
 	float4 CurrentPackedHeight = InSourceHeightmap.Load(int3(HeightmapCoordinates, 0));
@@ -113,11 +121,11 @@ void ApplyLandscapeTextureHeightPatch(in float4 SVPos : SV_POSITION, out float2
 	float NewHeight = 0;
 	if (bAdditiveMode)
 	{
-		NewHeight = CurrentHeight + Alpha * PatchHeight;
+		NewHeight = CurrentHeight + Alpha * PatchSignedHeight;
 	}
 	else
 	{
-		NewHeight = lerp(CurrentHeight, PatchHeight, Alpha);
+		NewHeight = lerp(CurrentHeight, LANDSCAPE_MID_VALUE + PatchSignedHeight, Alpha);
 	}
 	
 	OutColor = PackHeight(NewHeight);
@@ -130,11 +138,40 @@ float InHeightOffset;
 
 void ApplyOffsetToHeightmap(in float4 SVPos : SV_POSITION, out float2 OutColor : SV_Target0)
 {
-	float StoredHeight = UnpackHeight(InHeightmap.Load(int3(SVPos.xy, 0)));
+	float StoredHeight = UnpackHeight(InHeightmap.Load(int3(SVPos.xy, 0)).xy);
 	OutColor = PackHeight(StoredHeight + InHeightOffset);
 }
 #endif // OFFSET_LANDSCAPE_PATCH
 
+#if CONVERT_TO_NATIVE_LANDSCAPE_PATCH
+Texture2D<float4> InHeightmap;
+float InZeroInEncoding;
+float InHeightScale;
+float InHeightOffset;
+
+void ConvertToNativeLandscapePatch(in float4 SVPos : SV_POSITION, out float2 OutColor : SV_Target0)
+{
+	float StoredHeight = InHeightmap.Load(int3(SVPos.xy, 0)).x;
+	// Apply the same kinds of transformations that we would apply for alpha blend mode in ApplyLandscapeTextureHeightPatch
+	OutColor = PackHeight(InHeightScale * (StoredHeight - InZeroInEncoding) + InHeightOffset + LANDSCAPE_MID_VALUE);
+}
+#endif // CONVERT_TO_NATIVE_LANDSCAPE_PATCH
+
+#if CONVERT_BACK_FROM_NATIVE_LANDSCAPE_PATCH
+Texture2D<float4> InHeightmap;
+float InZeroInEncoding;
+float InHeightScale;
+float InHeightOffset;
+
+void ConvertBackFromNativeLandscapePatch(in float4 SVPos : SV_POSITION, out float OutColor : SV_Target0)
+{
+	float StoredHeight = UnpackHeight(InHeightmap.Load(int3(SVPos.xy, 0)).xy);
+	// Undo the transformation we do in ConvertToNativeLandscapePatch
+	OutColor = InHeightScale ? (StoredHeight - LANDSCAPE_MID_VALUE - InHeightOffset) / InHeightScale + InZeroInEncoding
+		: InZeroInEncoding;
+}
+#endif // CONVERT_BACK_FROM_NATIVE_LANDSCAPE_PATCH
+
 #if SIMPLE_TEXTURE_COPY
 Texture2D<float4> InSource;
 SamplerState InSourceSampler;
diff --git a/Engine/Plugins/Experimental/LandscapePatch/Source/LandscapePatch/Private/LandscapeTextureHeightPatch.cpp b/Engine/Plugins/Experimental/LandscapePatch/Source/LandscapePatch/Private/LandscapeTextureHeightPatch.cpp
--- a/Engine/Plugins/Experimental/LandscapePatch/Source/LandscapePatch/Private/LandscapeTextureHeightPatch.cpp
+++ b/Engine/Plugins/Experimental/LandscapePatch/Source/LandscapePatch/Private/LandscapeTextureHeightPatch.cpp
@@ -35,9 +35,10 @@ namespace LandscapeTextureHeightPatchLocals
 	}
 }
 
-void ULandscapeTextureHeightPatch::CopyInternalTextureToInternalRenderTarget()
+void ULandscapeTextureHeightPatch::ConvertInternalRenderTargetBackFromNativeTexture(bool bLoading)
 {
 	using namespace LandscapeTextureHeightPatchLocals;
+	using namespace UE::Landscape;
 
 	if (IsValid(InternalTexture))
 	{
@@ -66,17 +67,46 @@ void ULandscapeTextureHeightPatch::CopyInternalTextureToInternalRenderTarget()
 
 	FTextureResource* Source = InternalTexture->GetResource();
 	FTextureResource* Destination = InternalRenderTarget->GetResource();
-	ENQUEUE_RENDER_COMMAND(LandscapeTextureHeightPatchRTToTexture)(
-		[Source, Destination](FRHICommandListImmediate& RHICmdList)
+
+	// If we're in a different format, we need to "un-bake" the height from the texture.
+	if (InternalRenderTarget->RenderTargetFormat != ETextureRenderTargetFormat::RTF_RGBA8)
 	{
-		CopyTextureOnRenderThread(RHICmdList, *Source, *Destination);
-	});
+		FConvertToNativeLandscapePatchParams ConversionParams = GetConversionParams();
+		if (bLoading)
+		{
+			ConversionParams.HeightScale = SavedConversionHeightScale;
+		}
+
+		ENQUEUE_RENDER_COMMAND(LandscapeTextureHeightPatchRTToTexture)(
+			[ConversionParams, Source, Destination](FRHICommandListImmediate& RHICmdList)
+		{
+			using namespace UE::Landscape;
+
+			FRDGBuilder GraphBuilder(RHICmdList, RDG_EVENT_NAME("LandscapeTextureHeightPatchConvertFromNative"));
+
+			FRDGTextureRef SourceTexture = GraphBuilder.RegisterExternalTexture(CreateRenderTarget(Source->GetTexture2DRHI(), TEXT("ConversionSource")));
+			FRDGTextureRef DestinationTexture = GraphBuilder.RegisterExternalTexture(CreateRenderTarget(Destination->GetTexture2DRHI(), TEXT("ConversionDestination")));
+
+			FConvertBackFromNativeLandscapePatchPS::AddToRenderGraph(GraphBuilder, SourceTexture, DestinationTexture, ConversionParams);
 
+			GraphBuilder.Execute();
+		});
+	}
+	else
+	{
+		// When formats match, we can just copy back and forth.
+		ENQUEUE_RENDER_COMMAND(LandscapeTextureHeightPatchRTToTexture)(
+			[Source, Destination](FRHICommandListImmediate& RHICmdList)
+		{
+			CopyTextureOnRenderThread(RHICmdList, *Source, *Destination);
+		});
+	}
 }
 
-void ULandscapeTextureHeightPatch::CopyInternalRenderTargetToInternalTexture(bool bBlock)
+void ULandscapeTextureHeightPatch::ConvertInternalRenderTargetToNativeTexture(bool bBlock)
 {
 	using namespace LandscapeTextureHeightPatchLocals;
+	using namespace UE::Landscape;
 
 	if (!IsValid(InternalRenderTarget))
 	{
@@ -94,11 +124,48 @@ void ULandscapeTextureHeightPatch::CopyInternalRenderTargetToInternalTexture(boo
 	}
 	ResizeTextureIfNeeded(InternalRenderTarget->SizeX, InternalRenderTarget->SizeY, /*bClear*/ false, /*bUpdateResource*/ false);
 
-	// Read from the render target.
+	UTextureRenderTarget2D* NativeEncodingRenderTarget = InternalRenderTarget;
+
+	// If the format doesn't match the format that we use generally for our internal texture, save the patch in our native
+	// height format, applying whatever scale/offset is relevant. The stored texture thus ends up being the native equivalent
+	// (with scale 1 and offset 0). This is easier than trying to support various kinds of RT-to-texture conversions.
+	if (NativeEncodingRenderTarget->RenderTargetFormat != ETextureRenderTargetFormat::RTF_RGBA8)
+	{
+		// We need a temporary render target to write the converted result, then we'll copy that to the texture.
+		NativeEncodingRenderTarget = NewObject<UTextureRenderTarget2D>(this);
+		NativeEncodingRenderTarget->ClearColor = ClearColor;
+		NativeEncodingRenderTarget->RenderTargetFormat = ETextureRenderTargetFormat::RTF_RGBA8;
+		NativeEncodingRenderTarget->InitAutoFormat(InternalRenderTarget->SizeX, InternalRenderTarget->SizeY);
+		NativeEncodingRenderTarget->UpdateResourceImmediate(false);
+
+		FTextureResource* Source = InternalRenderTarget->GetResource();
+		FTextureResource* Destination = NativeEncodingRenderTarget->GetResource();
+
+		FConvertToNativeLandscapePatchParams ConversionParams = GetConversionParams();
+		SavedConversionHeightScale = ConversionParams.HeightScale;
+
+		ENQUEUE_RENDER_COMMAND(LandscapeTextureHeightPatchRTToTexture)(
+			[ConversionParams, Source, Destination](FRHICommandListImmediate& RHICmdList)
+		{
+			using namespace UE::Landscape;
+
+			FRDGBuilder GraphBuilder(RHICmdList, RDG_EVENT_NAME("LandscapeTextureHeightPatchConvertToNative"));
+
+			FRDGTextureRef SourceTexture = GraphBuilder.RegisterExternalTexture(CreateRenderTarget(Source->GetTexture2DRHI(), TEXT("ConversionSource")));
+			FRDGTextureRef DestinationTexture = GraphBuilder.RegisterExternalTexture(CreateRenderTarget(Destination->GetTexture2DRHI(), TEXT("ConversionDestination")));
+
+			FConvertToNativeLandscapePatchPS::AddToRenderGraph(GraphBuilder, SourceTexture, DestinationTexture, ConversionParams);
+
+			GraphBuilder.Execute();
+		});
+	}
+
+	// Write the render target to the texture.
 	// TODO: The header for this requires the texture to be square power of 2, but it actually doesn't seem to
 	// be an enforced requirement. If that changes, we'll need our own ReadPixels followed by locking a mip
 	// and writing to it.
-	InternalRenderTarget->UpdateTexture2D(InternalTexture, ETextureSourceFormat::TSF_BGRA8);
+	// This call does a flush for us, so the render target should be updated.
+	NativeEncodingRenderTarget->UpdateTexture2D(InternalTexture, ETextureSourceFormat::TSF_BGRA8);
 
 	InternalTexture->UpdateResource();
 	
@@ -108,15 +175,38 @@ void ULandscapeTextureHeightPatch::CopyInternalRenderTargetToInternalTexture(boo
 	}
 }
 
+UE::Landscape::FConvertToNativeLandscapePatchParams ULandscapeTextureHeightPatch::GetConversionParams()
+{
+	// When doing conversions, we bake into a height in the same way that we do when applying the patch.
+
+	UE::Landscape::FConvertToNativeLandscapePatchParams ConversionParams;
+	ConversionParams.ZeroInEncoding = EncodingSettings.ZeroInEncoding;
+
+	double LandscapeHeightScale = Landscape.IsValid() ? Landscape->GetTransform().GetScale3D().Z : 1;
+	LandscapeHeightScale = LandscapeHeightScale == 0 ? 1 : LandscapeHeightScale;
+	ConversionParams.HeightScale = EncodingSettings.WorldSpaceEncodingScale * LANDSCAPE_INV_ZSCALE / LandscapeHeightScale;
+
+	// TODO: We can choose whether we want to bake in the height offset if it exists. Doing so will handle
+	// some edge cases where the value stored in the patch is outside the range storeable in the native format
+	// normally, but within the range of the landscape due to the patch being far above/below the landscape to
+	// compensate. However, while this is good for conversions for the purposes of serialization, it's not good
+	// for conversions for the purposes of source mode change, so we would need to do things slightly differently
+	// in the two cases. For now, we'll just not bother with that (unlikely?) edge case.
+	ConversionParams.HeightOffset = 0;
+
+	return ConversionParams;
+}
+
 // Render targets don't get serialized, so whenever we need to save, copy, etc, we convert
 // to a UTexture2D, and then we convert back when needed.
 void ULandscapeTextureHeightPatch::PostLoad()
 {
 	Super::PostLoad();
 
-	if (SourceMode == ELandscapeTexturePatchSourceMode::TextureBackedRenderTarget)
+	if (SourceMode == ELandscapeTexturePatchSourceMode::TextureBackedRenderTarget && InternalTexture)
 	{
-		CopyInternalTextureToInternalRenderTarget();
+		InternalTexture->ConditionalPostLoad();
+		ConvertInternalRenderTargetBackFromNativeTexture(/* bLoading =*/ true);
 	}
 }
 
@@ -126,7 +216,7 @@ void ULandscapeTextureHeightPatch::PreSave(FObjectPreSaveContext SaveContext)
 
 	if (SourceMode == ELandscapeTexturePatchSourceMode::TextureBackedRenderTarget)
 	{
-		CopyInternalRenderTargetToInternalTexture(true);
+		ConvertInternalRenderTargetToNativeTexture(true);
 	}
 }
 
@@ -136,7 +226,7 @@ void ULandscapeTextureHeightPatch::PreDuplicate(FObjectDuplicationParameters& Du
 
 	if (SourceMode == ELandscapeTexturePatchSourceMode::TextureBackedRenderTarget)
 	{
-		CopyInternalRenderTargetToInternalTexture(true);
+		ConvertInternalRenderTargetToNativeTexture(true);
 	}
 }
 
@@ -147,7 +237,7 @@ void ULandscapeTextureHeightPatch::ExportCustomProperties(FOutputDevice& Out, ui
 
 	if (SourceMode == ELandscapeTexturePatchSourceMode::TextureBackedRenderTarget)
 	{
-		CopyInternalRenderTargetToInternalTexture(true);
+		ConvertInternalRenderTargetToNativeTexture(true);
 	}
 }
 
@@ -157,12 +247,20 @@ void ULandscapeTextureHeightPatch::PostEditChangeProperty(FPropertyChangedEvent&
 	{
 		if (PropertyChangedEvent.Property->GetFName() == GET_MEMBER_NAME_CHECKED(ULandscapeTextureHeightPatch, SourceMode))
 		{
-			SetSourceMode(SourceMode);
+			ResetSourceMode(SourceMode);
+			if (SourceMode == ELandscapeTexturePatchSourceMode::InternalTexture)
+			{
+				FTextureCompilingManager::Get().FinishCompilation({ InternalTexture });
+			}
 		}
 		else if (PropertyChangedEvent.Property->GetFName() == GET_MEMBER_NAME_CHECKED(ULandscapeTextureHeightPatch, InitializationMode))
 		{
 			bShowTextureAssetProperty = InitializationMode == ELandscapeTextureHeightPatchInitMode::TextureAsset;
 		}
+		else if (PropertyChangedEvent.Property->GetFName() == GET_MEMBER_NAME_CHECKED(ULandscapeTextureHeightPatch, SourceEncoding))
+		{
+			ResetSourceEncodingMode(SourceEncoding);
+		}
 	}
 		
 	Super::PostEditChangeProperty(PropertyChangedEvent);
@@ -186,7 +284,7 @@ void ULandscapeTextureHeightPatch::OnComponentCreated()
 		if (SourceMode == ELandscapeTexturePatchSourceMode::TextureBackedRenderTarget
 			&& IsValid(InternalTexture))
 		{
-			CopyInternalTextureToInternalRenderTarget();
+			ConvertInternalRenderTargetBackFromNativeTexture();
 		}
 	}
 	else // if not copy, ie adding a totally new component
@@ -206,11 +304,38 @@ void ULandscapeTextureHeightPatch::OnComponentCreated()
 	PreviousSourceMode = SourceMode;
 }
 
+void ULandscapeTextureHeightPatch::ResetSourceEncodingMode(ELandscapeTextureHeightPatchEncoding EncodingMode)
+{
+	Modify();
+	SourceEncoding = EncodingMode;
+	if (EncodingMode == ELandscapeTextureHeightPatchEncoding::ZeroToOne)
+	{
+		EncodingSettings.ZeroInEncoding = 0.5;
+		EncodingSettings.WorldSpaceEncodingScale = 400;
+	}
+	else if (EncodingMode == ELandscapeTextureHeightPatchEncoding::WorldUnits)
+	{
+		EncodingSettings.ZeroInEncoding = 0;
+		EncodingSettings.WorldSpaceEncodingScale = 1;
+	}
+}
+
+void ULandscapeTextureHeightPatch::SetInternalRenderTargetFormat(ETextureRenderTargetFormat Format)
+{
+	Modify();
+	InternalRenderTargetFormat = Format;
+	if (InternalRenderTarget)
+	{
+		ResizeRenderTargetIfNeeded(InternalRenderTarget->SizeX, InternalRenderTarget->SizeY);
+	}
+}
+
 void ULandscapeTextureHeightPatch::UpdateShaderParams(
 	UE::Landscape::FApplyLandscapeTextureHeightPatchPS::FParameters& Params, 
 	const FIntPoint& DestinationResolution, FIntRect& DestinationBoundsOut) const
 {
 	using namespace UE::Geometry;
+	using namespace UE::Landscape;
 
 	// We want our patch to be oriented with its Z axis to be along the Z axis of the landscape. The way we do this here
 	// is by just changing the rotation component of the patch transform to be the rotation of the landscape, except for
@@ -239,17 +364,28 @@ void ULandscapeTextureHeightPatch::UpdateShaderParams(
 	Params.InHeightmapToPatch = (FMatrix44f)LandscapeToPatchUVTransposed.GetTransposed();
 
 	FVector3d ComponentScale = PatchToWorld.GetScale3D();
-	float CombinedHeightScale = PatchHeightScale * (bApplyComponentZScale ? ComponentScale.Z : 1);
-	Params.InHeightScale = CombinedHeightScale;
+	double LandscapeHeightScale = Landscape.IsValid() ? Landscape->GetTransform().GetScale3D().Z : 1;
+	LandscapeHeightScale = LandscapeHeightScale == 0 ? 1 : LandscapeHeightScale;
+
+	bool bNativeEncoding = SourceMode == ELandscapeTexturePatchSourceMode::InternalTexture 
+		|| SourceEncoding == ELandscapeTextureHeightPatchEncoding::NativePackedHeight;
 
-	Params.InPatchBaseHeight = 0;
+	// To get height scale in heightmap coordinates, we have to undo the scaling that happens to map the 16bit int to [-256, 256), and undo
+	// the landscape actor scale.
+	Params.InHeightScale = bNativeEncoding ? 1 
+		: LANDSCAPE_INV_ZSCALE * EncodingSettings.WorldSpaceEncodingScale / LandscapeHeightScale;
+	if (bApplyComponentZScale)
+	{
+		Params.InHeightScale *= ComponentScale.Z;
+	}
+
+	Params.InZeroInEncoding = bNativeEncoding ? LandscapeDataAccess::MidValue : EncodingSettings.ZeroInEncoding;
+
+	Params.InHeightOffset = 0;
 	if (bUsePatchZAsReference)
 	{
 		FVector3d PatchOriginInHeightmapCoords = LandscapeHeightmapToWorld.InverseTransformPosition(PatchToWorld.GetTranslation());
-
-		Params.InPatchBaseHeight = CombinedHeightScale < 0 ?
-			PatchOriginInHeightmapCoords.Z - (LandscapeDataAccess::MaxValue - LandscapeDataAccess::MidValue) * CombinedHeightScale
-			: PatchOriginInHeightmapCoords.Z - LandscapeDataAccess::MidValue * CombinedHeightScale;
+		Params.InHeightOffset = PatchOriginInHeightmapCoords.Z - LandscapeDataAccess::MidValue;
 	}
 
 	// The outer half-pixel shouldn't affect the landscape because it is not part of our official coverage area.
@@ -260,11 +396,25 @@ void ULandscapeTextureHeightPatch::UpdateShaderParams(
 		Params.InEdgeUVDeadBorder = FVector2f(0.5 / TextureResolution.X, 0.5 / TextureResolution.Y);
 	}
 
-	Params.bRectangularFalloff = FalloffMode == ELandscapeTextureHeightPatchFalloffMode::RoundedRectangle;
 	Params.InFalloffWorldMargin = Falloff / FMath::Min(ComponentScale.X, ComponentScale.Y);
 
-	Params.bApplyPatchAlpha = bUseTextureAlphaChannel;
-	Params.bAdditiveMode = BlendMode == ELandscapeTextureHeightPatchBlendMode::Additive;
+	// Pack our booleans into a bitfield
+	using EShaderFlags = FApplyLandscapeTextureHeightPatchPS::EFlags;
+	EShaderFlags Flags = EShaderFlags::None;
+
+	Flags |= (FalloffMode == ELandscapeTextureHeightPatchFalloffMode::RoundedRectangle) ?
+		EShaderFlags::RectangularFalloff : EShaderFlags::None;
+
+	Flags |= bUseTextureAlphaChannel ?
+		EShaderFlags::ApplyPatchAlpha : EShaderFlags::None;
+
+	Flags |= (BlendMode == ELandscapeTextureHeightPatchBlendMode::Additive) ?
+		EShaderFlags::AdditiveMode : EShaderFlags::None;
+
+	Flags |= bNativeEncoding ?
+		EShaderFlags::InputIsPackedHeight : EShaderFlags::None;
+
+	Params.InFlags = static_cast<uint8>(Flags);
 
 	// Get the output bounds, which are used to limit the amount of landscape pixels we have to process. 
 	// To get them, convert all of the corners into heightmap 2d coordinates and get the bounding box.
@@ -296,7 +446,8 @@ UTextureRenderTarget2D* ULandscapeTextureHeightPatch::Render_Native(bool bIsHeig
 
 	if (!ensure(PatchManager.IsValid()) || 
 		(SourceMode == ELandscapeTexturePatchSourceMode::TextureBackedRenderTarget && (!IsValid(InternalRenderTarget)))
-		|| (SourceMode == ELandscapeTexturePatchSourceMode::TextureAsset && (!IsValid(TextureAsset)))
+		|| (SourceMode == ELandscapeTexturePatchSourceMode::TextureAsset && (!IsValid(TextureAsset) 
+			|| !ensureMsgf(TextureAsset->VirtualTextureStreaming == 0, TEXT("ULandscapeTextureHeightPatch: Virtual textures are not supported"))))
 		|| (SourceMode == ELandscapeTexturePatchSourceMode::InternalTexture && (!IsValid(InternalTexture))))
 	{
 		return InCombinedResult;
@@ -404,7 +555,9 @@ void ULandscapeTextureHeightPatch::Reinitialize()
 			"to be a texture asset."));
 		return;
 	}
-	if ((InitializationMode == ELandscapeTextureHeightPatchInitMode::TextureAsset && !IsValid(TextureAsset))
+	if ((InitializationMode == ELandscapeTextureHeightPatchInitMode::TextureAsset 
+			&& (!IsValid(TextureAsset) 
+				|| !ensureMsgf(TextureAsset->VirtualTextureStreaming == 0, TEXT("ULandscapeTextureHeightPatch: Virtual textures are not supported"))))
 		|| (InitializationMode == ELandscapeTextureHeightPatchInitMode::FromLandscape && !Landscape.IsValid()))
 	{
 		// Don't have what we need for initialization
@@ -441,6 +594,7 @@ void ULandscapeTextureHeightPatch::Reinitialize()
 	switch (InitializationMode)
 	{
 	case ELandscapeTextureHeightPatchInitMode::Blank:
+		SourceEncoding = ELandscapeTextureHeightPatchEncoding::NativePackedHeight;
 		if (SourceMode == ELandscapeTexturePatchSourceMode::TextureBackedRenderTarget)
 		{
 			ResizeRenderTargetIfNeeded(InitTextureSizeX, InitTextureSizeY);
@@ -462,6 +616,7 @@ void ULandscapeTextureHeightPatch::Reinitialize()
 
 		// We're going to need our render target regardless of source mode because we'll write to it 
 		// before copying to texture if needed.
+		SourceEncoding = ELandscapeTextureHeightPatchEncoding::NativePackedHeight;
 		ResizeRenderTargetIfNeeded(InitTextureSizeX, InitTextureSizeY);
 			
 		// If bUsePatchZAsReference is true, then we're going to be adding an offset to our rendered landscape
@@ -482,10 +637,7 @@ void ULandscapeTextureHeightPatch::Reinitialize()
 
 		Landscape->RenderHeightmap(GetPatchToWorldTransform(), FBox2D(-FullPatchDimensions / 2, FullPatchDimensions / 2), RenderedHeightmapSection);
 
-		// Reset height related parameters because adjusting for them would be a pain and doesn't
-		// seem useful to the user if they are reinitializing...
 		bApplyComponentZScale = false;
-		PatchHeightScale = 1;
 
 		// ...However, we do need to apply the offset to account for our patch's position relative to the
 		// landscape z axis, if that is a setting that the user is using.
@@ -521,7 +673,7 @@ void ULandscapeTextureHeightPatch::Reinitialize()
 
 		if (SourceMode == ELandscapeTexturePatchSourceMode::InternalTexture)
 		{
-			CopyInternalRenderTargetToInternalTexture(true);
+			ConvertInternalRenderTargetToNativeTexture(true);
 			InternalRenderTarget = nullptr;
 		}
 
@@ -558,7 +710,7 @@ void ULandscapeTextureHeightPatch::Reinitialize()
 
 		if (SourceMode == ELandscapeTexturePatchSourceMode::InternalTexture)
 		{
-			CopyInternalRenderTargetToInternalTexture(true);
+			ConvertInternalRenderTargetToNativeTexture(true);
 			InternalRenderTarget = nullptr;
 		}
 		break;
@@ -577,6 +729,9 @@ bool ULandscapeTextureHeightPatch::ResizeRenderTargetIfNeeded(int32 SizeX, int32
 
 	bool bChanged = false;
 
+	ETextureRenderTargetFormat FormatToUse = SourceEncoding == ELandscapeTextureHeightPatchEncoding::NativePackedHeight ?
+		ETextureRenderTargetFormat::RTF_RGBA8 : InternalRenderTargetFormat.GetValue();
+
 	if (!IsValid(InternalRenderTarget))
 	{
 		Modify();
@@ -584,16 +739,17 @@ bool ULandscapeTextureHeightPatch::ResizeRenderTargetIfNeeded(int32 SizeX, int32
 		InternalRenderTarget = NewObject<UTextureRenderTarget2D>(this);
 		InternalRenderTarget->ClearColor = ClearColor;
 
-		InternalRenderTarget->RenderTargetFormat = ETextureRenderTargetFormat::RTF_RGBA8;
+		InternalRenderTarget->RenderTargetFormat = FormatToUse;
 		InternalRenderTarget->InitAutoFormat(SizeX, SizeY);
 
 		bChanged = true;
 	}
 	else if (InternalRenderTarget->SizeX != SizeX
-		|| InternalRenderTarget->SizeY != SizeY)
+		|| InternalRenderTarget->SizeY != SizeY
+		|| InternalRenderTarget->RenderTargetFormat != FormatToUse)
 	{
 		InternalRenderTarget->Modify();
-		InternalRenderTarget->RenderTargetFormat = ETextureRenderTargetFormat::RTF_RGBA8;
+		InternalRenderTarget->RenderTargetFormat = FormatToUse;
 		InternalRenderTarget->InitAutoFormat(SizeX, SizeY);
 
 		bChanged = true;
@@ -675,12 +831,12 @@ bool ULandscapeTextureHeightPatch::SetSourceMode(ELandscapeTexturePatchSourceMod
 	if (PreviousSourceMode == ELandscapeTexturePatchSourceMode::TextureBackedRenderTarget
 		&& NewMode == ELandscapeTexturePatchSourceMode::InternalTexture)
 	{
-		CopyInternalRenderTargetToInternalTexture(false);
+		ConvertInternalRenderTargetToNativeTexture(false);
 	}
 	else if (PreviousSourceMode == ELandscapeTexturePatchSourceMode::InternalTexture
 		&& NewMode == ELandscapeTexturePatchSourceMode::TextureBackedRenderTarget)
 	{
-		CopyInternalTextureToInternalRenderTarget();
+		ConvertInternalRenderTargetBackFromNativeTexture();
 	}
 
 	if (bDeleteUnusedInternalTextures)
@@ -700,6 +856,27 @@ bool ULandscapeTextureHeightPatch::SetSourceMode(ELandscapeTexturePatchSourceMod
 	return true;
 }
 
+bool ULandscapeTextureHeightPatch::ResetSourceMode(ELandscapeTexturePatchSourceMode NewMode)
+{
+	if (NewMode == ELandscapeTexturePatchSourceMode::TextureAsset)
+	{
+		ResetSourceEncodingMode(ELandscapeTextureHeightPatchEncoding::ZeroToOne);
+	}
+	else if (NewMode == ELandscapeTexturePatchSourceMode::TextureBackedRenderTarget)
+	{
+		ResetSourceEncodingMode(ELandscapeTextureHeightPatchEncoding::WorldUnits);
+	}
+	// The internal texture case ignores the encoding mode property so it does not need resetting. More
+	// importantly, not resetting it allows us to use the current encoding settings to convert properly
+	// into the native format inside SetSourceMode, which calls ConvertInternalRenderTargetToNativeTexture.
+	//else
+	//{
+	//	ResetSourceEncodingMode(ELandscapeTextureHeightPatchEncoding::NativePackedHeight);
+	//}
+
+	return SetSourceMode(NewMode, true);
+}
+
 void ULandscapeTextureHeightPatch::SnapToLandscape()
 {
 	Modify();
diff --git a/Engine/Plugins/Experimental/LandscapePatch/Source/LandscapePatch/Private/LandscapeTextureHeightPatchPS.cpp b/Engine/Plugins/Experimental/LandscapePatch/Source/LandscapePatch/Private/LandscapeTextureHeightPatchPS.cpp
--- a/Engine/Plugins/Experimental/LandscapePatch/Source/LandscapePatch/Private/LandscapeTextureHeightPatchPS.cpp
+++ b/Engine/Plugins/Experimental/LandscapePatch/Source/LandscapePatch/Private/LandscapeTextureHeightPatchPS.cpp
@@ -16,6 +16,12 @@ bool FApplyLandscapeTextureHeightPatchPS::ShouldCompilePermutation(const FGlobal
 void FApplyLandscapeTextureHeightPatchPS::ModifyCompilationEnvironment(const FGlobalShaderPermutationParameters& Parameters, FShaderCompilerEnvironment& OutEnvironment)
 {
 	OutEnvironment.SetDefine(TEXT("APPLY_LANDSCAPE_PATCH"), 1);
+
+	// Make our flag choices match in the shader.
+	OutEnvironment.SetDefine(TEXT("RECTANGULAR_FALLOFF_FLAG"), static_cast<uint8>(FApplyLandscapeTextureHeightPatchPS::EFlags::RectangularFalloff));
+	OutEnvironment.SetDefine(TEXT("APPLY_PATCH_ALPHA_FLAG"), static_cast<uint8>(FApplyLandscapeTextureHeightPatchPS::EFlags::ApplyPatchAlpha));
+	OutEnvironment.SetDefine(TEXT("ADDITIVE_MODE_FLAG"), static_cast<uint8>(FApplyLandscapeTextureHeightPatchPS::EFlags::AdditiveMode));
+	OutEnvironment.SetDefine(TEXT("INPUT_IS_PACKED_HEIGHT_FLAG"), static_cast<uint8>(FApplyLandscapeTextureHeightPatchPS::EFlags::InputIsPackedHeight));
 }
 
 void FApplyLandscapeTextureHeightPatchPS::AddToRenderGraph(FRDGBuilder& GraphBuilder, FParameters* InParameters, const FIntRect& DestinationBounds)
@@ -91,8 +97,80 @@ void FSimpleTextureCopyPS::AddToRenderGraph(FRDGBuilder& GraphBuilder, FRDGTextu
 		FIntRect(0, 0, DestinationSize.X, DestinationSize.Y));
 }
 
+bool FConvertToNativeLandscapePatchPS::ShouldCompilePermutation(const FGlobalShaderPermutationParameters& Parameters)
+{
+	return DoesPlatformSupportEditLayers(Parameters.Platform);
+}
+
+void FConvertToNativeLandscapePatchPS::ModifyCompilationEnvironment(const FGlobalShaderPermutationParameters& Parameters, FShaderCompilerEnvironment& OutEnvironment)
+{
+	OutEnvironment.SetDefine(TEXT("CONVERT_TO_NATIVE_LANDSCAPE_PATCH"), 1);
+}
+
+void FConvertToNativeLandscapePatchPS::AddToRenderGraph(FRDGBuilder& GraphBuilder, FRDGTextureRef SourceTexture, 
+	FRDGTextureRef DestinationTexture, const FConvertToNativeLandscapePatchParams& Params)
+{
+	FGlobalShaderMap* ShaderMap = GetGlobalShaderMap(GMaxRHIFeatureLevel);
+	TShaderMapRef<FConvertToNativeLandscapePatchPS> PixelShader(ShaderMap);
+
+	FRDGTextureSRVRef SourceSRV = GraphBuilder.CreateSRV(FRDGTextureSRVDesc::CreateForMipLevel(SourceTexture, 0));
+
+	FConvertToNativeLandscapePatchPS::FParameters* ShaderParams = GraphBuilder.AllocParameters<FConvertToNativeLandscapePatchPS::FParameters>();
+	ShaderParams->InHeightmap = GraphBuilder.CreateSRV(FRDGTextureSRVDesc::CreateForMipLevel(SourceTexture, 0));
+	ShaderParams->InZeroInEncoding = Params.ZeroInEncoding;
+	ShaderParams->InHeightScale = Params.HeightScale;
+	ShaderParams->InHeightOffset = Params.HeightOffset;
+	ShaderParams->RenderTargets[0] = FRenderTargetBinding(DestinationTexture, ERenderTargetLoadAction::ENoAction, /*InMipIndex = */0);
+	FIntVector DestinationSize = DestinationTexture->Desc.GetSize();
+
+	FPixelShaderUtils::AddFullscreenPass(
+		GraphBuilder,
+		ShaderMap,
+		RDG_EVENT_NAME("ConvertToNativeLandscapePatch"),
+		PixelShader,
+		ShaderParams,
+		FIntRect(0, 0, DestinationSize.X, DestinationSize.Y));
+}
+
+bool FConvertBackFromNativeLandscapePatchPS::ShouldCompilePermutation(const FGlobalShaderPermutationParameters& Parameters)
+{
+	return DoesPlatformSupportEditLayers(Parameters.Platform);
+}
+
+void FConvertBackFromNativeLandscapePatchPS::ModifyCompilationEnvironment(const FGlobalShaderPermutationParameters& Parameters, FShaderCompilerEnvironment& OutEnvironment)
+{
+	OutEnvironment.SetDefine(TEXT("CONVERT_BACK_FROM_NATIVE_LANDSCAPE_PATCH"), 1);
+}
+
+void FConvertBackFromNativeLandscapePatchPS::AddToRenderGraph(FRDGBuilder& GraphBuilder, FRDGTextureRef SourceTexture,
+	FRDGTextureRef DestinationTexture, const FConvertToNativeLandscapePatchParams& Params)
+{
+	FGlobalShaderMap* ShaderMap = GetGlobalShaderMap(GMaxRHIFeatureLevel);
+	TShaderMapRef<FConvertBackFromNativeLandscapePatchPS> PixelShader(ShaderMap);
+
+	FRDGTextureSRVRef SourceSRV = GraphBuilder.CreateSRV(FRDGTextureSRVDesc::CreateForMipLevel(SourceTexture, 0));
+
+	FConvertBackFromNativeLandscapePatchPS::FParameters* ShaderParams = GraphBuilder.AllocParameters<FConvertBackFromNativeLandscapePatchPS::FParameters>();
+	ShaderParams->InHeightmap = GraphBuilder.CreateSRV(FRDGTextureSRVDesc::CreateForMipLevel(SourceTexture, 0));
+	ShaderParams->InZeroInEncoding = Params.ZeroInEncoding;
+	ShaderParams->InHeightScale = Params.HeightScale;
+	ShaderParams->InHeightOffset = Params.HeightOffset;
+	ShaderParams->RenderTargets[0] = FRenderTargetBinding(DestinationTexture, ERenderTargetLoadAction::ENoAction, /*InMipIndex = */0);
+	FIntVector DestinationSize = DestinationTexture->Desc.GetSize();
+
+	FPixelShaderUtils::AddFullscreenPass(
+		GraphBuilder,
+		ShaderMap,
+		RDG_EVENT_NAME("ConvertBackFromNativeLandscapePatch"),
+		PixelShader,
+		ShaderParams,
+		FIntRect(0, 0, DestinationSize.X, DestinationSize.Y));
+}
+
 }//end UE::Landscape
 
 IMPLEMENT_GLOBAL_SHADER(UE::Landscape::FApplyLandscapeTextureHeightPatchPS, "/Plugin/LandscapePatch/Private/LandscapeTextureHeightPatchPS.usf", "ApplyLandscapeTextureHeightPatch", SF_Pixel);
 IMPLEMENT_GLOBAL_SHADER(UE::Landscape::FOffsetHeightmapPS, "/Plugin/LandscapePatch/Private/LandscapeTextureHeightPatchPS.usf", "ApplyOffsetToHeightmap", SF_Pixel);
 IMPLEMENT_GLOBAL_SHADER(UE::Landscape::FSimpleTextureCopyPS, "/Plugin/LandscapePatch/Private/LandscapeTextureHeightPatchPS.usf", "SimpleTextureCopy", SF_Pixel);
+IMPLEMENT_GLOBAL_SHADER(UE::Landscape::FConvertToNativeLandscapePatchPS, "/Plugin/LandscapePatch/Private/LandscapeTextureHeightPatchPS.usf", "ConvertToNativeLandscapePatch", SF_Pixel);
+IMPLEMENT_GLOBAL_SHADER(UE::Landscape::FConvertBackFromNativeLandscapePatchPS, "/Plugin/LandscapePatch/Private/LandscapeTextureHeightPatchPS.usf", "ConvertBackFromNativeLandscapePatch", SF_Pixel);
diff --git a/Engine/Plugins/Experimental/LandscapePatch/Source/LandscapePatch/Private/LandscapeTexturePatchBase.cpp b/Engine/Plugins/Experimental/LandscapePatch/Source/LandscapePatch/Private/LandscapeTexturePatchBase.cpp
--- a/Engine/Plugins/Experimental/LandscapePatch/Source/LandscapePatch/Private/LandscapeTexturePatchBase.cpp
+++ b/Engine/Plugins/Experimental/LandscapePatch/Source/LandscapePatch/Private/LandscapeTexturePatchBase.cpp
@@ -7,6 +7,13 @@
 #include "Engine/TextureRenderTarget2D.h"
 #include "Landscape.h"
 
+void ULandscapeTexturePatchBase::SetTextureAsset(UTexture* TextureIn)
+{
+	ensureMsgf(!TextureIn || TextureIn->VirtualTextureStreaming == 0, 
+		TEXT("ULandscapeTexturePatchBase::SetTextureAsset: Virtual textures are not supported."));
+	TextureAsset = TextureIn; 
+}
+
 FTransform ULandscapeTexturePatchBase::GetPatchToWorldTransform() const
 {
 	FTransform PatchToWorld = GetComponentTransform();
diff --git a/Engine/Plugins/Experimental/LandscapePatch/Source/LandscapePatch/Public/LandscapeTextureHeightPatch.h b/Engine/Plugins/Experimental/LandscapePatch/Source/LandscapePatch/Public/LandscapeTextureHeightPatch.h
--- a/Engine/Plugins/Experimental/LandscapePatch/Source/LandscapePatch/Public/LandscapeTextureHeightPatch.h
+++ b/Engine/Plugins/Experimental/LandscapePatch/Source/LandscapePatch/Public/LandscapeTextureHeightPatch.h
@@ -57,6 +57,49 @@ enum class ELandscapeTextureHeightPatchFalloffMode : uint8
 	RoundedRectangle,
 };
 
+UENUM(BlueprintType)
+enum class ELandscapeTextureHeightPatchEncoding : uint8
+{
+	// Values in texture should be interpreted as being floats in the range [0,1]. User specifies what
+	// value corresponds to height 0 (i.e. height when landscape is "cleared"), and the size of the 
+	// range in world units.
+	ZeroToOne,
+
+	// Values in texture are direct world-space heights.
+	WorldUnits,
+
+	// Values in texture are stored the same way they are in landscape actors: as 16 bit integers packed 
+	// into two bytes, mapping to [-256, 256 - 1/128] before applying landscape scale.
+	NativePackedHeight
+
+	//~ Note that currently ZeroToOne and WorldUnits actually work the same way- we subtract the center point (0 for WorldUnits),
+	//~ then scale in some way (1.0 for WorldUnits). However, having separate options here allows us to initialize defaults
+	//~ appropriately when setting the encoding mode.
+};
+
+//~ A struct in case we find that we need other encoding settings.
+USTRUCT(BlueprintType)
+struct LANDSCAPEPATCH_API FLandscapeTexturePatchEncodingSettings
+{
+	GENERATED_BODY()
+public:
+	/**
+	 * The value in the patch data that corresponds to 0 landscape height (which is in line with patch Z when
+	 * "Use Patch Z As Reference" is true, and at landscape zero/mid value when false).
+	 */
+	UPROPERTY(EditAnywhere, BlueprintReadWrite, Category = Settings)
+	double ZeroInEncoding = 0;
+
+	/**
+	 * The scale that should be aplied to the data stored in the patch relative to the zero in the encoding, in world coordinates.
+	 * For instance if the encoding is [0,1], and 0.5 correponds to 0, a WorldSpaceEncoding Scale of 100 means that the resulting
+	 * values will lie in the range [-50, 50] in world space, which would be [-0.5, 0.5] in the landscape local heights if the Z
+	 * scale is 100.
+	 */
+	UPROPERTY(EditAnywhere, BlueprintReadWrite, Category = Settings)
+	double WorldSpaceEncodingScale = 1;
+};
+
 /**
  * A texture-based height patch
  */
@@ -76,6 +119,12 @@ public:
 	// ULandscapeTexturePatchBase
 	virtual bool SetTextureResolution(FVector2D ResolutionIn) override;
 	virtual bool SetSourceMode(ELandscapeTexturePatchSourceMode NewMode, bool bDeleteUnusedInternalTextures = true) override;
+	
+	/**
+	 * Just like SetSourceMode, but resets encoding mode and settings to mode-specific defaults
+	 */
+	UFUNCTION(BlueprintCallable, Category = LandscapePatch)
+	bool ResetSourceMode(ELandscapeTexturePatchSourceMode NewMode);
 
 	// UActorComponent
 	virtual void OnComponentCreated() override;
@@ -111,36 +160,93 @@ public:
 	 */
 	UFUNCTION(BlueprintCallable, CallInEditor, Category = Initialization)
 	void SnapToLandscape();
+
+	/**
+	 * Changes the render target format of the internal render target. This will usually result
+	 * in a clearing the render target since it usually has to be rebuild.
+	 */
+	UFUNCTION(BlueprintCallable, Category = LandscapePatch, meta = (ETextureRenderTargetFormat = "ETextureRenderTargetFormat::RTF_R32f"))
+	void SetInternalRenderTargetFormat(ETextureRenderTargetFormat Format);
+
+	/**
+	 * Set the height encoding mode for the patch, which determines how stored values in the patch
+	 * are translated into heights when applying to landscape. 
+	 */
+	UFUNCTION(BlueprintCallable, Category = LandscapePatch)
+	void SetSourceEncodingMode(ELandscapeTextureHeightPatchEncoding EncodingMode) 
+	{ 
+		Modify();
+		SourceEncoding = EncodingMode; 
+	}
+
+	/**
+	 * Just like SetSourceEncodingMode, but resets ZeroInEncoding and WorldSpaceEncodingScale to mode-specific defaults.
+	 */
+	UFUNCTION(BlueprintCallable, Category = LandscapePatch)
+	void ResetSourceEncodingMode(ELandscapeTextureHeightPatchEncoding EncodingMode);
+
+	/**
+	 * Set settings that determine how values in the patch are translated into heights. This is only
+	 * used if the encoding mode is not NativePackedHeight, where values are expected to be already
+	 * in the same space as the landscape heightmap.
+	 */
+	UFUNCTION(BlueprintCallable, Category = LandscapePatch)
+	void SetEncodingSettings(const FLandscapeTexturePatchEncodingSettings& Settings)
+	{
+		Modify();
+		EncodingSettings = Settings;
+	}
 #endif
 
 	UFUNCTION(BlueprintCallable, Category = "LandscapePatch")
-	void SetFalloff(float FalloffIn) { Falloff = FalloffIn; }
+	void SetFalloff(float FalloffIn) 
+	{ 
+		Modify();
+		Falloff = FalloffIn; 
+	}
 
 	UFUNCTION(BlueprintCallable, Category = "LandscapePatch")
-	void SetBlendMode(ELandscapeTextureHeightPatchBlendMode BlendModeIn) { BlendMode = BlendModeIn; }
+	void SetBlendMode(ELandscapeTextureHeightPatchBlendMode BlendModeIn) 
+	{ 
+		Modify();
+		BlendMode = BlendModeIn; 
+	}
 
 	/**
 	 * Determine whether the patch Z relative to the landscape affects the results. See comment for bUsePatchZAsReference.
 	 */
 	UFUNCTION(BlueprintCallable, Category = "LandscapePatch")
-	void SetUsePatchZAsReference(bool bUsePatchZIn) { bUsePatchZAsReference = bUsePatchZIn; }
+	void SetUsePatchZAsReference(bool bUsePatchZIn) 
+	{ 
+		Modify();
+		bUsePatchZAsReference = bUsePatchZIn; 
+	}
 
 	UFUNCTION(BlueprintCallable, Category = "LandscapePatch")
-	void SetUseTextureAlphaChannel(bool bUse) { bUseTextureAlphaChannel = bUse; }
+	void SetUseTextureAlphaChannel(bool bUse) 
+	{ 
+		Modify();
+		bUseTextureAlphaChannel = bUse; 
+	}
+
 
 protected:
 
-	/**
-	 * Extra scaling to apply to patch height. Useful when using adjusting existing texture assets.
-	 */
-	UPROPERTY(EditAnywhere, Category = Settings, meta = (UIMin = "0", UIMax = "1"))
-	float PatchHeightScale = 1;
+	UPROPERTY(EditAnywhere, Category = Settings, AdvancedDisplay)
+	ELandscapeTextureHeightPatchBlendMode BlendMode = ELandscapeTextureHeightPatchBlendMode::AlphaBlend;
+	
+	/** When true, texture alpha channel will be used when applying the patch. */
+	UPROPERTY(EditAnywhere, Category = Settings, AdvancedDisplay)
+	bool bUseTextureAlphaChannel = false;
+
+	UPROPERTY(EditAnywhere, Category = Settings)
+	ELandscapeTextureHeightPatchFalloffMode FalloffMode = ELandscapeTextureHeightPatchFalloffMode::RoundedRectangle;
 
 	/**
-	 * Whether to apply the patch Z scale to the height stored in the patch.
+	 * Distance (in unscaled world coordinates) across which to smoothly fall off the patch effects.
 	 */
-	UPROPERTY(EditAnywhere, Category = Settings, AdvancedDisplay, meta = (DisplayName = "Apply Component Z Scale"))
-	bool bApplyComponentZScale = true;
+	UPROPERTY(EditAnywhere, Category = Settings, meta = (ClampMin = "0", UIMax = "2000"))
+	float Falloff = 0;
 
 	/**
 	 * Determine whether the patch Z relative to the landscape affects the results. When true, moving the
@@ -154,22 +260,24 @@ protected:
 	UPROPERTY(EditAnywhere, Category = Settings, AdvancedDisplay)
 	bool bUsePatchZAsReference = true;
 
-	UPROPERTY(EditAnywhere, Category = Settings, AdvancedDisplay)
-	ELandscapeTextureHeightPatchBlendMode BlendMode = ELandscapeTextureHeightPatchBlendMode::AlphaBlend;
-	
-	/** When true, texture alpha channel will be used when applying the patch. */
-	UPROPERTY(EditAnywhere, Category = Settings, AdvancedDisplay)
-	bool bUseTextureAlphaChannel = false;
+	/** How the values stored in the patch represent the height. */
+	UPROPERTY(EditAnywhere, Category = Settings, AdvancedDisplay, meta = (EditConditionHides,
+		EditCondition = "SourceMode != ELandscapeTexturePatchSourceMode::InternalTexture"))
+	ELandscapeTextureHeightPatchEncoding SourceEncoding = ELandscapeTextureHeightPatchEncoding::NativePackedHeight;
 
-	UPROPERTY(EditAnywhere, Category = Settings)
-	ELandscapeTextureHeightPatchFalloffMode FalloffMode = ELandscapeTextureHeightPatchFalloffMode::RoundedRectangle;
+	UPROPERTY(EditAnywhere, Category = Settings, AdvancedDisplay, meta = (UIMin = "0", UIMax = "1", EditConditionHides,
+		EditCondition = "SourceMode != ELandscapeTexturePatchSourceMode::InternalTexture && SourceEncoding != ELandscapeTextureHeightPatchEncoding::NativePackedHeight"))
+	FLandscapeTexturePatchEncodingSettings EncodingSettings;
 
 	/**
-	 * Distance (in unscaled world coordinates) across which to smoothly fall off the patch effects.
+	 * Whether to apply the patch Z scale to the height stored in the patch.
 	 */
-	UPROPERTY(EditAnywhere, Category = Settings, meta = (ClampMin = "0", UIMax = "2000"))
-	float Falloff = 0;
+	UPROPERTY(EditAnywhere, Category = Settings, AdvancedDisplay, meta = (DisplayName = "Apply Component Z Scale"))
+	bool bApplyComponentZScale = true;
 
+	/**
+	 * Controls how the patch is initialized when invoking Reinitialize().
+	 */
 	UPROPERTY(EditAnywhere, Category = Initialization, meta = (
 		EditCondition = "SourceMode != ELandscapeTexturePatchSourceMode::TextureAsset", EditConditionHides))
 	ELandscapeTextureHeightPatchInitMode InitializationMode = ELandscapeTextureHeightPatchInitMode::FromLandscape;
@@ -214,10 +322,24 @@ protected:
 
 private:
 
+	// Uprops can't be inside WITH_EDITOR, apparently. Could put in WITH_EDITOR_DATA.
+	UPROPERTY()
+	TEnumAsByte<ETextureRenderTargetFormat> InternalRenderTargetFormat = ETextureRenderTargetFormat::RTF_R32f;
+
+	// When loading the patch, the landscape may not be available to provide the scale, so we save
+	// it in case we need it.
+	// TODO: It would be cleaner to save all conversion parameters but then we have to make FConvertToNativeLandscapePatchParams
+	// into a UStruct, etc. Revisit if needed later, but for now definitely keep private.
+	UPROPERTY()
+	float SavedConversionHeightScale;
+
 #if WITH_EDITOR
-	// Uncertain whether these may end up in the base class once we have weight patches.
-	void CopyInternalRenderTargetToInternalTexture(bool bBlock);
-	void CopyInternalTextureToInternalRenderTarget();
+	// These are used to convert render targets whose formats are not the usual two-channel 16bit values that
+	// we store in an internal texture. We convert these render targets to the native landscape height representation
+	// and store that in an internal texture.
+	UE::Landscape::FConvertToNativeLandscapePatchParams GetConversionParams();
+	void ConvertInternalRenderTargetToNativeTexture(bool bBlock);
+	void ConvertInternalRenderTargetBackFromNativeTexture(bool bLoading = false);
 
 	void UpdateShaderParams(UE::Landscape::FApplyLandscapeTextureHeightPatchPS::FParameters& Params, 
 		const FIntPoint& DestinationResolution, FIntRect& DestinationBoundsOut) const;
diff --git a/Engine/Plugins/Experimental/LandscapePatch/Source/LandscapePatch/Public/LandscapeTextureHeightPatchPS.h b/Engine/Plugins/Experimental/LandscapePatch/Source/LandscapePatch/Public/LandscapeTextureHeightPatchPS.h
--- a/Engine/Plugins/Experimental/LandscapePatch/Source/LandscapePatch/Public/LandscapeTextureHeightPatchPS.h
+++ b/Engine/Plugins/Experimental/LandscapePatch/Source/LandscapePatch/Public/LandscapeTextureHeightPatchPS.h
@@ -20,26 +20,46 @@ class LANDSCAPEPATCH_API FApplyLandscapeTextureHeightPatchPS : public FGlobalSha
 
 public:
 
+	// Flags that get packed into a bitfield because we're not allowed to use bool shader parameters:
+	enum class EFlags : uint8
+	{
+		None = 0,
+
+		RectangularFalloff = 1 << 0,
+		ApplyPatchAlpha = 1 << 1,
+		AdditiveMode = 1 << 2,
+		InputIsPackedHeight = 1 << 3
+	};
+
+	// When false, falloff is circular.
+	static const uint32 RectangularFalloffFlag = 1 << 0;
+	// When true, the texture alpha channel is considered for blending (in addition to falloff, if nonzero)
+	static const uint32 ApplyPatchAlphaFlag = 1 << 1;
+	// When true, the patch is added to the landscape rather than alpha blending with the landscape.
+	static const uint32 AdditiveModeFlag = 1 << 2;
+	// When false, the input is directly interpreted as being the height value to process. When true, the height
+	// is unpacked from the red and green channels to make a 16 bit int.
+	static const uint32 InputIsPackedHeightFlag = 1 << 3;
+
 	BEGIN_SHADER_PARAMETER_STRUCT(FParameters, )
 		SHADER_PARAMETER_RDG_TEXTURE_SRV(Texture2D<float4>, InSourceHeightmap)
 		SHADER_PARAMETER_RDG_TEXTURE_SRV(Texture2D<float4>, InHeightPatch)
 		SHADER_PARAMETER_SAMPLER(SamplerState, InHeightPatchSampler)
 		SHADER_PARAMETER(FMatrix44f, InHeightmapToPatch)
-		// Height value to consider the 0 (deepest) values in height patch
-		SHADER_PARAMETER(float, InPatchBaseHeight)
+		// Value in patch that corresponds to the landscape mid value, which is our "0 height".
+		SHADER_PARAMETER(float, InZeroInEncoding)
+		// Scale to apply to source values relative to the value that represents 0 height.
+		SHADER_PARAMETER(float, InHeightScale)
+		// Offset to apply to height result after applying height scale
+		SHADER_PARAMETER(float, InHeightOffset)
 		// Amount of the patch edge to not apply in UV space. Generally set to 0.5/Dimensions to avoid applying
 		// the edge half-pixels.
 		SHADER_PARAMETER(FVector2f, InEdgeUVDeadBorder)
 		// In patch texture space, the size of the margin across which the alpha falls from 1 to 0
 		SHADER_PARAMETER(float, InFalloffWorldMargin)
 		SHADER_PARAMETER(FVector2f, InPatchWorldDimensions)
-		SHADER_PARAMETER(float, InHeightScale)
-
-		// TODO: Apparently we're not allowed to use bools, and we're theoretically supposed to pack these into a bitfield.
-		// Will be done in later CL.
-		SHADER_PARAMETER(uint32, bRectangularFalloff)
-		SHADER_PARAMETER(uint32, bApplyPatchAlpha)
-		SHADER_PARAMETER(uint32, bAdditiveMode)
+		// Some combination of the flags (see constants above).
+		SHADER_PARAMETER(uint32, InFlags)
 
 		RENDER_TARGET_BINDING_SLOTS() // Holds our output
 		END_SHADER_PARAMETER_STRUCT()
@@ -49,6 +69,8 @@ public:
 	static void AddToRenderGraph(FRDGBuilder& GraphBuilder, FParameters* InParameters, const FIntRect& DestinationBounds);
 };
 
+ENUM_CLASS_FLAGS(FApplyLandscapeTextureHeightPatchPS::EFlags);
+
 /**
  * Simple shader that just offsets each height value in a height patch by a constant.
  */
@@ -70,6 +92,66 @@ public:
 	static void AddToRenderGraph(FRDGBuilder& GraphBuilder, FParameters* InParameters);
 };
 
+// Values needed to convert a patch stored in some source encoding into the native (two byte int) encoding and back
+struct FConvertToNativeLandscapePatchParams
+{
+	float ZeroInEncoding;
+	float HeightScale;
+	float HeightOffset;
+};
+
+/**
+ * Shader that converts a texture stored in some other encoding (where height is in the R channel) to the 
+ * landscape "native" encoding, where height is stored as a 16 bit int split across the R and G channels.
+ * This is not perfectly reversible (in case of clamping and due to rounding), but it lets us store the
+ * texture in the way that it would be applied to the landscape (usually).
+ */
+class LANDSCAPEPATCH_API FConvertToNativeLandscapePatchPS : public FGlobalShader
+{
+	DECLARE_GLOBAL_SHADER(FConvertToNativeLandscapePatchPS);
+	SHADER_USE_PARAMETER_STRUCT(FConvertToNativeLandscapePatchPS, FGlobalShader);
+
+public:
+
+	BEGIN_SHADER_PARAMETER_STRUCT(FParameters, )
+		SHADER_PARAMETER_RDG_TEXTURE_SRV(Texture2D<float4>, InHeightmap)
+		SHADER_PARAMETER(float, InZeroInEncoding)
+		SHADER_PARAMETER(float, InHeightScale)
+		SHADER_PARAMETER(float, InHeightOffset)
+		RENDER_TARGET_BINDING_SLOTS() // Holds our output
+		END_SHADER_PARAMETER_STRUCT()
+
+	static bool ShouldCompilePermutation(const FGlobalShaderPermutationParameters& Parameters);
+	static void ModifyCompilationEnvironment(const FGlobalShaderPermutationParameters& InParameters, FShaderCompilerEnvironment& OutEnvironment);
+	static void AddToRenderGraph(FRDGBuilder& GraphBuilder, FRDGTextureRef SourceTexture,
+		FRDGTextureRef DestinationTexture, const FConvertToNativeLandscapePatchParams& Params);
+};
+
+/**
+ * Shader that undoes the conversion done by FConvertToNativeLandscapePatchPS (to the extent possible, since
+ * rounding/clamping makes it not perfectly recoverable).
+ */
+class LANDSCAPEPATCH_API FConvertBackFromNativeLandscapePatchPS : public FGlobalShader
+{
+	DECLARE_GLOBAL_SHADER(FConvertBackFromNativeLandscapePatchPS);
+	SHADER_USE_PARAMETER_STRUCT(FConvertBackFromNativeLandscapePatchPS, FGlobalShader);
+
+public:
+
+	BEGIN_SHADER_PARAMETER_STRUCT(FParameters, )
+		SHADER_PARAMETER_RDG_TEXTURE_SRV(Texture2D<float4>, InHeightmap)
+		SHADER_PARAMETER(float, InZeroInEncoding)
+		SHADER_PARAMETER(float, InHeightScale)
+		SHADER_PARAMETER(float, InHeightOffset)
+		RENDER_TARGET_BINDING_SLOTS() // Holds our output
+		END_SHADER_PARAMETER_STRUCT()
+
+	static bool ShouldCompilePermutation(const FGlobalShaderPermutationParameters& Parameters);
+	static void ModifyCompilationEnvironment(const FGlobalShaderPermutationParameters& InParameters, FShaderCompilerEnvironment& OutEnvironment);
+	static void AddToRenderGraph(FRDGBuilder& GraphBuilder, FRDGTextureRef SourceTexture,
+		FRDGTextureRef DestinationTexture, const FConvertToNativeLandscapePatchParams& Params);
+};
+
 /**
  * Simple shader for copying textures of potentially different resolutions.
  */
diff --git a/Engine/Plugins/Experimental/LandscapePatch/Source/LandscapePatch/Public/LandscapeTexturePatchBase.h b/Engine/Plugins/Experimental/LandscapePatch/Source/LandscapePatch/Public/LandscapeTexturePatchBase.h
--- a/Engine/Plugins/Experimental/LandscapePatch/Source/LandscapePatch/Public/LandscapeTexturePatchBase.h
+++ b/Engine/Plugins/Experimental/LandscapePatch/Source/LandscapePatch/Public/LandscapeTexturePatchBase.h
@@ -58,7 +58,7 @@ public:
 
 	// These need to be public so that we can take the internal texture and write it to an external one
 	UTexture2D* GetInternalTexture() { return InternalTexture; }
-	void SetTextureAsset(UTexture* TextureIn) { TextureAsset = TextureIn; }
+	void SetTextureAsset(UTexture* TextureIn);
 
 	virtual ELandscapeTexturePatchSourceMode GetSourceMode() const { return SourceMode; }
 	/**
@@ -138,7 +138,8 @@ protected:
 	 * Texture used when source mode is set to a texture asset, or for intializing from a texture asset.
 	 */
 	UPROPERTY(EditAnywhere, Category = Settings, meta = (EditConditionHides, DisplayPriority = 20, DisplayName = "Texture Asset", 
-		EditCondition = "SourceMode == ELandscapeTexturePatchSourceMode::TextureAsset || bShowTextureAssetProperty"))
+		EditCondition = "SourceMode == ELandscapeTexturePatchSourceMode::TextureAsset || bShowTextureAssetProperty", 
+		DisallowedAssetDataTags = "VirtualTextureStreaming=True"))
 	TObjectPtr<UTexture> TextureAsset = nullptr;
 
 	/** At scale 1.0, the X and Y of the region affected by the height patch. This corresponds to the distance from the center
