LandscapePatch: Make most functions usable in non-editor-only actors, where they generally don't do anything. Make mode change on saving a patch undoable. Fix size of loaded RT patch sometimes being wrong.

diff --git a/Engine/Plugins/Experimental/LandscapePatch/Source/LandscapePatch/Private/LandscapeTextureHeightPatch.cpp b/Engine/Plugins/Experimental/LandscapePatch/Source/LandscapePatch/Private/LandscapeTextureHeightPatch.cpp
--- a/Engine/Plugins/Experimental/LandscapePatch/Source/LandscapePatch/Private/LandscapeTextureHeightPatch.cpp
+++ b/Engine/Plugins/Experimental/LandscapePatch/Source/LandscapePatch/Private/LandscapeTextureHeightPatch.cpp
@@ -60,11 +60,12 @@ void ULandscapeTextureHeightPatch::ConvertInternalRenderTargetBackFromNativeText
 	{
 		InternalRenderTarget->Modify();
 	}
-	ResizeRenderTargetIfNeeded(InternalTexture->GetResource()->GetSizeX(), 
-		InternalTexture->GetResource()->GetSizeY());
 
 	FTextureCompilingManager::Get().FinishCompilation({ InternalTexture });
 
+	ResizeRenderTargetIfNeeded(InternalTexture->GetResource()->GetSizeX(), 
+		InternalTexture->GetResource()->GetSizeY());
+
 	FTextureResource* Source = InternalTexture->GetResource();
 	FTextureResource* Destination = InternalRenderTarget->GetResource();
 
@@ -303,6 +304,7 @@ void ULandscapeTextureHeightPatch::OnComponentCreated()
 
 	PreviousSourceMode = SourceMode;
 }
+#endif // WITH_EDITOR
 
 void ULandscapeTextureHeightPatch::ResetSourceEncodingMode(ELandscapeTextureHeightPatchEncoding EncodingMode)
 {
@@ -322,14 +324,17 @@ void ULandscapeTextureHeightPatch::ResetSourceEncodingMode(ELandscapeTextureHeig
 
 void ULandscapeTextureHeightPatch::SetInternalRenderTargetFormat(ETextureRenderTargetFormat Format)
 {
+#if WITH_EDITOR
 	Modify();
 	InternalRenderTargetFormat = Format;
 	if (InternalRenderTarget)
 	{
 		ResizeRenderTargetIfNeeded(InternalRenderTarget->SizeX, InternalRenderTarget->SizeY);
 	}
+#endif
 }
 
+#if WITH_EDITOR
 void ULandscapeTextureHeightPatch::UpdateShaderParams(
 	UE::Landscape::FApplyLandscapeTextureHeightPatchPS::FParameters& Params, 
 	const FIntPoint& DestinationResolution, FIntRect& DestinationBoundsOut) const
@@ -526,6 +531,7 @@ UTextureRenderTarget2D* ULandscapeTextureHeightPatch::Render_Native(bool bIsHeig
 
 	return InCombinedResult;
 }
+#endif // WITH_EDITOR
 
 void ULandscapeTextureHeightPatch::DeleteInternalTextures()
 {
@@ -545,6 +551,7 @@ void ULandscapeTextureHeightPatch::DeleteInternalTextures()
 
 void ULandscapeTextureHeightPatch::Reinitialize()
 {
+#if WITH_EDITOR
 	using namespace UE::Landscape;
 	using namespace LandscapeTextureHeightPatchLocals;
 
@@ -721,8 +728,10 @@ void ULandscapeTextureHeightPatch::Reinitialize()
 	{
 		PatchManager->RequestLandscapeUpdate();
 	}
+#endif // WITH_EDITOR
 }
 
+#if WITH_EDITOR
 bool ULandscapeTextureHeightPatch::ResizeRenderTargetIfNeeded(int32 SizeX, int32 SizeY)
 {
 	using namespace LandscapeTextureHeightPatchLocals;
@@ -855,6 +864,7 @@ bool ULandscapeTextureHeightPatch::SetSourceMode(ELandscapeTexturePatchSourceMod
 
 	return true;
 }
+#endif // WITH_EDITOR
 
 bool ULandscapeTextureHeightPatch::ResetSourceMode(ELandscapeTexturePatchSourceMode NewMode)
 {
@@ -971,6 +981,7 @@ void ULandscapeTextureHeightPatch::SnapToLandscape()
 	SetWorldLocation(LandscapeTransform.TransformPosition(NewCenterInLandscape));
 }
 
+#if WITH_EDITOR
 bool ULandscapeTextureHeightPatch::SetTextureResolution(FVector2D ResolutionIn)
 {
 	if (SourceMode == ELandscapeTexturePatchSourceMode::TextureAsset)
@@ -997,5 +1008,4 @@ bool ULandscapeTextureHeightPatch::SetTextureResolution(FVector2D ResolutionIn)
 }
 #endif // WITH_EDITOR
 
-
 #undef LOCTEXT_NAMESPACE
\ No newline at end of file
diff --git a/Engine/Plugins/Experimental/LandscapePatch/Source/LandscapePatch/Public/LandscapeTextureHeightPatch.h b/Engine/Plugins/Experimental/LandscapePatch/Source/LandscapePatch/Public/LandscapeTextureHeightPatch.h
--- a/Engine/Plugins/Experimental/LandscapePatch/Source/LandscapePatch/Public/LandscapeTextureHeightPatch.h
+++ b/Engine/Plugins/Experimental/LandscapePatch/Source/LandscapePatch/Public/LandscapeTextureHeightPatch.h
@@ -119,12 +119,6 @@ public:
 	// ULandscapeTexturePatchBase
 	virtual bool SetTextureResolution(FVector2D ResolutionIn) override;
 	virtual bool SetSourceMode(ELandscapeTexturePatchSourceMode NewMode, bool bDeleteUnusedInternalTextures = true) override;
-	
-	/**
-	 * Just like SetSourceMode, but resets encoding mode and settings to mode-specific defaults
-	 */
-	UFUNCTION(BlueprintCallable, Category = LandscapePatch)
-	bool ResetSourceMode(ELandscapeTexturePatchSourceMode NewMode);
 
 	// UActorComponent
 	virtual void OnComponentCreated() override;
@@ -136,6 +130,13 @@ public:
 	virtual void ExportCustomProperties(FOutputDevice& Out, uint32 Indent) override;
 	virtual void PostEditChangeProperty(FPropertyChangedEvent& PropertyChangedEvent) override;
 	virtual void PostInitProperties() override;
+#endif
+
+	/**
+	 * Just like SetSourceMode, but resets encoding mode and settings to mode-specific defaults
+	 */
+	UFUNCTION(BlueprintCallable, Category = LandscapePatch)
+	bool ResetSourceMode(ELandscapeTexturePatchSourceMode NewMode);
 
 	/**
 	 * Deletes the internal render target and internal texture.
@@ -196,11 +197,10 @@ public:
 		Modify();
 		EncodingSettings = Settings;
 	}
-#endif
 
 	UFUNCTION(BlueprintCallable, Category = "LandscapePatch")
 	void SetFalloff(float FalloffIn) 
-	{ 
+	{
 		Modify();
 		Falloff = FalloffIn; 
 	}
@@ -333,6 +333,10 @@ private:
 	UPROPERTY()
 	float SavedConversionHeightScale;
 
+	// Many of the functions above are actually also editor-only in that their contents are inside
+	// WITH_EDITOR blocks. However we don't put the actual signatures into WITH_EDITOR because that
+	// would prevent them from being called from non-editor-only blueprints, even if the call happens
+	// in a place that happens to be editor-only.
 #if WITH_EDITOR
 	// These are used to convert render targets whose formats are not the usual two-channel 16bit values that
 	// we store in an internal texture. We convert these render targets to the native landscape height representation
diff --git a/Engine/Plugins/Experimental/LandscapePatch/Source/LandscapePatchEditorOnly/Private/LandscapeTexturePatchCustomization.cpp b/Engine/Plugins/Experimental/LandscapePatch/Source/LandscapePatchEditorOnly/Private/LandscapeTexturePatchCustomization.cpp
--- a/Engine/Plugins/Experimental/LandscapePatch/Source/LandscapePatchEditorOnly/Private/LandscapeTexturePatchCustomization.cpp
+++ b/Engine/Plugins/Experimental/LandscapePatch/Source/LandscapePatchEditorOnly/Private/LandscapeTexturePatchCustomization.cpp
@@ -90,6 +90,8 @@ void FLandscapeTexturePatchCustomization::CustomizeDetails(IDetailLayoutBuilder&
 
 					if (SavedObjects.Num() > 0)
 					{
+						PatchObject->Modify();
+
 						UTexture* NewTexture = Cast<UTexture>(SavedObjects[0]);
 						if (ensure(NewTexture))
 						{
