LandscapePatch: Exposed more functions to non-editor-only actors.

diff --git a/Engine/Plugins/Experimental/LandscapePatch/Source/LandscapePatch/Public/LandscapePatchComponent.h b/Engine/Plugins/Experimental/LandscapePatch/Source/LandscapePatch/Public/LandscapePatchComponent.h
--- a/Engine/Plugins/Experimental/LandscapePatch/Source/LandscapePatch/Public/LandscapePatchComponent.h
+++ b/Engine/Plugins/Experimental/LandscapePatch/Source/LandscapePatch/Public/LandscapePatchComponent.h
@@ -33,15 +33,16 @@ public:
 		UTextureRenderTarget2D* InCombinedResult,
 		const FName& InWeightmapLayerName) { return InCombinedResult; }
 
-// For now we keep the patches largely editor-only, since we don't yet support runtime landscape editing
-#if WITH_EDITOR
-
 	UFUNCTION(BlueprintCallable, Category = "LandscapePatch")
 	void RequestLandscapeUpdate();
 
 	UFUNCTION(BlueprintCallable, Category = "LandscapePatch")
 	FTransform GetLandscapeHeightmapCoordsToWorld() const;
 
+	// For now we keep the patches largely editor-only, since we don't yet support runtime landscape editing.
+	// The above functions are also editor-only (and don't work at runtime), but can't be in WITH_EDITOR blocks
+	// so that they can be called from non-editor-only classes in editor contexts.
+#if WITH_EDITOR
 	// USceneComponent
 	virtual void OnUpdateTransform(EUpdateTransformFlags UpdateTransformFlags, ETeleportType Teleport) override;
 
