LandscapePatch: Make landscape patches work with one-file-per-actor. Add min/max blend modes. Allow user to specify patch manager that they want to use. Allow enabling/disabling patches temporarily. Place new patch layer under water layer by default. Don't change source encoding when changing source mode. Grey out most things instead of hiding them in detail panel, and move things out of Advanced Display. When setting source type to be render target, initialize it to a default resolution.

diff --git a/Engine/Plugins/Experimental/LandscapePatch/Shaders/Private/LandscapeTextureHeightPatchPS.usf b/Engine/Plugins/Experimental/LandscapePatch/Shaders/Private/LandscapeTextureHeightPatchPS.usf
--- a/Engine/Plugins/Experimental/LandscapePatch/Shaders/Private/LandscapeTextureHeightPatchPS.usf
+++ b/Engine/Plugins/Experimental/LandscapePatch/Shaders/Private/LandscapeTextureHeightPatchPS.usf
@@ -31,6 +31,8 @@ float InZeroInEncoding;
 float InHeightScale;
 // An offset to apply to the resulting height, after applying the height scale
 float InHeightOffset;
+// Enum defining blend modes, with values set in corresponding cpp file.
+uint InBlendMode;
 // A combination of flags, whose positions are set in the corresponding cpp file.
 uint InFlags;
 
@@ -44,7 +46,6 @@ void ApplyLandscapeTextureHeightPatch(in float4 SVPos : SV_POSITION, out float2
 	// so that they match to the cpp file
 	bool bRectangularFalloff = InFlags & RECTANGULAR_FALLOFF_FLAG;
 	bool bApplyPatchAlpha = InFlags & APPLY_PATCH_ALPHA_FLAG;
-	bool bAdditiveMode = InFlags & ADDITIVE_MODE_FLAG;
 	bool bInputIsPackedHeight = InFlags & INPUT_IS_PACKED_HEIGHT_FLAG;
 	
 	// We need only the 2D affine transformation that goes from landscape XY heightmap integer coordinates
@@ -119,13 +120,20 @@ void ApplyLandscapeTextureHeightPatch(in float4 SVPos : SV_POSITION, out float2
 	}
 	
 	float NewHeight = 0;
-	if (bAdditiveMode)
+	switch (InBlendMode)
 	{
-		NewHeight = CurrentHeight + Alpha * PatchSignedHeight;
-	}
-	else
-	{
-		NewHeight = lerp(CurrentHeight, LANDSCAPE_MID_VALUE + PatchSignedHeight, Alpha);
+		case ADDITIVE_MODE:
+			NewHeight = CurrentHeight + Alpha * PatchSignedHeight;
+			break;
+		case ALPHA_BLEND_MODE:
+			NewHeight = lerp(CurrentHeight, LANDSCAPE_MID_VALUE + PatchSignedHeight, Alpha);
+			break;
+		case MIN_MODE:
+			NewHeight = lerp(CurrentHeight, min(CurrentHeight, LANDSCAPE_MID_VALUE + PatchSignedHeight), Alpha);
+			break;
+		case MAX_MODE:
+			NewHeight = lerp(CurrentHeight, max(CurrentHeight, LANDSCAPE_MID_VALUE + PatchSignedHeight), Alpha);
+			break;
 	}
 	
 	OutColor = PackHeight(NewHeight);
diff --git a/Engine/Plugins/Experimental/LandscapePatch/Source/LandscapePatch/LandscapePatch.Build.cs b/Engine/Plugins/Experimental/LandscapePatch/Source/LandscapePatch/LandscapePatch.Build.cs
--- a/Engine/Plugins/Experimental/LandscapePatch/Source/LandscapePatch/LandscapePatch.Build.cs
+++ b/Engine/Plugins/Experimental/LandscapePatch/Source/LandscapePatch/LandscapePatch.Build.cs
@@ -39,6 +39,11 @@ public class LandscapePatch : ModuleRules
 				"RenderCore",
 				"RHI",
 				"Projects", // IPluginManager
+				
+				// This doesn't seem to be editor-only currently. If it becomes editor-only, may need
+				// to inherit from LanscapeBlueprintBrushBase instead, and will lose the ability to
+				// add the patch manager via the Blueprints tool in Landscape mode.
+				"LandscapeEditorUtilities",
 			}
 			);
 	}
diff --git a/Engine/Plugins/Experimental/LandscapePatch/Source/LandscapePatch/Private/LandscapePatchComponent.cpp b/Engine/Plugins/Experimental/LandscapePatch/Source/LandscapePatch/Private/LandscapePatchComponent.cpp
--- a/Engine/Plugins/Experimental/LandscapePatch/Source/LandscapePatch/Private/LandscapePatchComponent.cpp
+++ b/Engine/Plugins/Experimental/LandscapePatch/Source/LandscapePatch/Private/LandscapePatchComponent.cpp
@@ -10,9 +10,10 @@
 
 #define LOCTEXT_NAMESPACE "LandscapePatch"
 
-#if WITH_EDITOR
 namespace LandscapePatchComponentLocals
 {
+#if WITH_EDITOR
+
 	ALandscapePatchManager* CreateNewPatchManagerForLandscape(ALandscape* Landscape)
 	{
 		if (!ensure(Landscape->CanHaveLayersContent()))
@@ -35,8 +36,9 @@ namespace LandscapePatchComponentLocals
 
 		return PatchManager;
 	}
-}
+
 #endif
+}
 
 // Note that this is not allowed to be editor-only
 ULandscapePatchComponent::ULandscapePatchComponent(const FObjectInitializer& ObjectInitializer)
@@ -49,6 +51,17 @@ ULandscapePatchComponent::ULandscapePatchComponent(const FObjectInitializer& Obj
 	bWantsOnUpdateTransform = true;
 }
 
+void ULandscapePatchComponent::SetIsEnabled(bool bEnabledIn)
+{
+	if (bEnabledIn == bIsEnabled)
+	{
+		return;
+	}
+
+	bIsEnabled = bEnabledIn;
+	RequestLandscapeUpdate();
+}
+
 #if WITH_EDITOR
 void ULandscapePatchComponent::OnComponentCreated()
 {
@@ -76,7 +89,7 @@ void ULandscapePatchComponent::OnComponentCreated()
 			// If we copied over a patch with a landscape but no manager, create manager in that landscape
 			if (Landscape->CanHaveLayersContent())
 			{
-				PatchManager = CreateNewPatchManagerForLandscape(Landscape.Get());
+				SetPatchManager(CreateNewPatchManagerForLandscape(Landscape.Get()));
 			}
 			else
 			{
@@ -92,8 +105,7 @@ void ULandscapePatchComponent::OnComponentCreated()
 			TActorIterator<ALandscapePatchManager> ManagerIterator(World);
 			if (!!ManagerIterator)
 			{
-				PatchManager = *ManagerIterator;
-				Landscape = PatchManager->GetOwningLandscape();
+				SetPatchManager(*ManagerIterator);
 			}
 			else
 			{
@@ -103,7 +115,7 @@ void ULandscapePatchComponent::OnComponentCreated()
 					if (LandscapeIterator->CanHaveLayersContent())
 					{
 						Landscape = *LandscapeIterator;
-						PatchManager = CreateNewPatchManagerForLandscape(Landscape.Get());
+						SetPatchManager(CreateNewPatchManagerForLandscape(Landscape.Get()));
 						break;
 					}
 				}
@@ -181,60 +193,112 @@ void ULandscapePatchComponent::PostEditChangeProperty(FPropertyChangedEvent& Pro
 
 	Super::PostEditChangeProperty(PropertyChangedEvent);
 
-	// If we're changing the owning landscape, there's some work we need to do to remove/add ourselves from/to the proper
-	// brush managers.
+	// If we're changing the owning landscape or patch manaager, there's some work we need to do to remove/add 
+	// ourselves from/to the proper brush managers.
 	if (PropertyChangedEvent.Property 
 		&& (PropertyChangedEvent.Property->GetFName() == GET_MEMBER_NAME_CHECKED(ULandscapePatchComponent, Landscape)))
 	{
-		if (PatchManager.IsValid())
-		{
-			PatchManager->RemovePatch(this);
-		}
-		PatchManager = nullptr;
+		SetLandscape(Landscape.Get());
+	}
+	else if (PropertyChangedEvent.Property
+		&& (PropertyChangedEvent.Property->GetFName() == GET_MEMBER_NAME_CHECKED(ULandscapePatchComponent, PatchManager)))
+	{
+		SetPatchManager(PatchManager.Get());
+	}
 
-		bool bAdded = !Landscape.IsValid();
-		for (TActorIterator<ALandscapePatchManager> ManagersItr(GetWorld()); ManagersItr; ++ManagersItr)
-		{
-			// Exit early if found the relevant manager.
-			if (bAdded)
-			{
-				break;
-			}
+	RequestLandscapeUpdate();
+}
+#endif
 
-			ALandscapePatchManager* Manager = *ManagersItr;
-			if (!ensure(Manager->IsValidLowLevel()))
-			{
-				continue;
-			}
+void ULandscapePatchComponent::SetLandscape(ALandscape* NewLandscape)
+{
+#if WITH_EDITOR
 
-			if (!bAdded && Manager->GetOwningLandscape() == Landscape)
-			{
-				Manager->AddPatch(this);
-				PatchManager = Manager;
-				bAdded = true;
-			}
+	using namespace LandscapePatchComponentLocals;
+
+	// Uncertain whether we want this early-out here. Perhaps, like SetPatchManager, we want 
+	// to be able to use this function to reorder patches. On the other hand, it seems inconvenient
+	// to accidentally swap patch managers if there are multiple in the same landscape, and we kept
+	// landscape the same. It's hard to know the ideal behavior, but for now we'll keep it.
+	if (Landscape == NewLandscape && ((!Landscape && !PatchManager)
+		|| (Landscape && PatchManager && PatchManager->GetOwningLandscape() == Landscape)))
+	{
+		return;
+	}
+
+	Landscape = NewLandscape;
+
+	if (!NewLandscape)
+	{
+		SetPatchManager(nullptr);
+		return;
+	}
+
+	// If landscape was valid, try to find a patch manager inside that landscape.
+	bool bFoundExisting = false;
+	for (int32 LayerIndex = 0; LayerIndex < Landscape->Layers.Num(); ++LayerIndex)
+	{
+		TArray<ALandscapeBlueprintBrushBase*> Brushes = Landscape->GetBrushesForLayer(LayerIndex);
+		ALandscapeBlueprintBrushBase** ExistingManager = Brushes.FindByPredicate(
+			[](ALandscapeBlueprintBrushBase* Brush) { return Cast<ALandscapePatchManager>(Brush) != nullptr; });
+		if (ExistingManager)
+		{
+			bFoundExisting = true;
+			SetPatchManager(Cast<ALandscapePatchManager>(*ExistingManager));
+			break;
 		}
+	}
 
-		// If we had a non-null landscape and didn't add ourselves to its manager despite iterating through
-		// all managers, then the landscape must not have had a manager. Create one for it.
-		if (!bAdded)
+	// If we didn't find an existing manager, create one for this landscape.
+	if (!bFoundExisting)
+	{
+		if (Landscape->CanHaveLayersContent())
 		{
-			if (Landscape->CanHaveLayersContent())
-			{
-				PatchManager = CreateNewPatchManagerForLandscape(Landscape.Get());
-				PatchManager->AddPatch(this);
-			}
-			else
-			{
-				UE_LOG(LogLandscapePatch, Warning, TEXT("Landscape target for height patch did not have edit layers enabled. Unable to create patch manager."));
-				Landscape = nullptr;
-			}
+			SetPatchManager(CreateNewPatchManagerForLandscape(Landscape.Get()));
+			PatchManager->AddPatch(this);
+		}
+		else
+		{
+			UE_LOG(LogLandscapePatch, Warning, TEXT("Landscape target for height patch did not have edit layers enabled. Unable to create patch manager."));
+			Landscape = nullptr;
+			SetPatchManager(nullptr);
 		}
 	}
 
-	RequestLandscapeUpdate();
+#endif // WITH_EDITOR
+}
+
+void ULandscapePatchComponent::SetPatchManager(ALandscapePatchManager* NewPatchManager)
+{
+	// Could have put the WITH_EDITOR block around just GetOwningLandscape, but might as well
+	// surround everything since nothing is currently expected to work for a patch manager at runtime.
+#if WITH_EDITOR
+
+	// TODO: We don't currently have an early out here (for PreviousPatchManager == NewPatchManager) 
+	// because using SetPatchManager is currently a convenient way to reorder patches inside the manager,
+	// and because we want to make sure the patch is added to the manager if it was removed from
+	// the list somehow. However we may want to revisit this if it turns out to cause problems.
+
+	if (PreviousPatchManager.IsValid())
+	{
+		PreviousPatchManager->RemovePatch(this);
+	}
+
+	PatchManager = NewPatchManager;
+	if (NewPatchManager)
+	{
+		PatchManager->AddPatch(this);
+		Landscape = PatchManager->GetOwningLandscape();
+	}
+	else
+	{
+		Landscape = nullptr;
+	}
+
+	PreviousPatchManager = NewPatchManager;
+
+#endif // WITH_EDITOR
 }
-#endif
 
 void ULandscapePatchComponent::RequestLandscapeUpdate()
 {
diff --git a/Engine/Plugins/Experimental/LandscapePatch/Source/LandscapePatch/Private/LandscapePatchManager.cpp b/Engine/Plugins/Experimental/LandscapePatch/Source/LandscapePatch/Private/LandscapePatchManager.cpp
--- a/Engine/Plugins/Experimental/LandscapePatch/Source/LandscapePatch/Private/LandscapePatchManager.cpp
+++ b/Engine/Plugins/Experimental/LandscapePatch/Source/LandscapePatch/Private/LandscapePatchManager.cpp
@@ -2,17 +2,18 @@
 
 #include "LandscapePatchManager.h"
 
-#include "Landscape.h"
-#include "LandscapePatchComponent.h"
 #include "Engine/TextureRenderTarget2D.h"
+#include "Landscape.h"
 #include "LandscapeDataAccess.h"
+#include "LandscapePatchComponent.h"
+#include "LandscapePatchLogging.h"
 #include "RenderGraph.h" // RDG_EVENT_NAME
 
 #define LOCTEXT_NAMESPACE "LandscapePatchManager"
 
 // TODO: Not sure if using this kind of constructor is a proper thing to do vs some other hook...
 ALandscapePatchManager::ALandscapePatchManager(const FObjectInitializer& ObjectInitializer)
-	: ALandscapeBlueprintBrushBase(ObjectInitializer)
+	: ALandscapeBlueprintBrush(ObjectInitializer)
 {
 #if WITH_EDITOR
 	SetAffectsHeightmap(true);
@@ -78,24 +79,40 @@ UTextureRenderTarget2D* ALandscapePatchManager::Render_Native(bool InIsHeightmap
 	// height map.
 	// So for now we do the simplest thing, and that is to have the height patches act as if they were
 	// independent brushes.
-	for (TWeakObjectPtr<ULandscapePatchComponent>& Component : PatchComponents)
+	for (TSoftObjectPtr<ULandscapePatchComponent>& Component : PatchComponents)
 	{
-		// Theoretically when components are marked for destruction, they should remove themselves from
-		// the patch manager, hence the ensure here
-		if (ensure(Component.IsValid()))
+		if (Component.IsPending())
 		{
-			InCombinedResult = Component->Render_Native(InIsHeightmap, InCombinedResult, InWeightmapLayerName);
+			Component.LoadSynchronous();
 		}
-		else
+
+		if (Component.IsValid())
+		{
+			if (Component->IsEnabled())
+			{
+				InCombinedResult = Component->Render_Native(InIsHeightmap, InCombinedResult, InWeightmapLayerName);
+			}
+		}
+		else if (Component.IsNull())
 		{
+			// Theoretically when components are marked for destruction, they should remove themselves from
+			// the patch manager in their OnComponentDestroyed call. However there seem to be ways to end up
+			// with destroyed patches not being removed, for instance through saving the manager but not the
+			// patch actor.
+			UE_LOG(LogLandscapePatch, Warning, TEXT("ALandscapePatchManager: Found an invalid patch in patch manager. It will be removed."));
 			bHaveInvalidPatches = true;
 		}
+		else
+		{
+			// This means that IsPending() was true, but LoadSynchronous() failed, which we don't expect to happen.
+			ensure(false);
+		}
 	}
 
 	if (bHaveInvalidPatches)
 	{
-		PatchComponents.RemoveAll([](TWeakObjectPtr<ULandscapePatchComponent> Component) {
-			return !Component.IsValid();
+		PatchComponents.RemoveAll([](TSoftObjectPtr<ULandscapePatchComponent> Component) {
+			return Component.IsNull();
 		});
 	}
 
@@ -119,6 +136,34 @@ void ALandscapePatchManager::SetTargetLandscape(ALandscape* InTargetLandscape)
 			if (ExistingPatchLayerIndex == INDEX_NONE)
 			{
 				ExistingPatchLayerIndex = InTargetLandscape->CreateLayer(LayerName);
+
+				// If there's a layer with a water brush manager, put the layer we created under that
+				// TODO: It's uncertain whether we want this approach long term. We might want to have a popup ask where to
+				// insert the layer, for instance. However, this is what the first artists to test the module suggested,
+				// and it's fairly easy to swap layers around, so it is easy for users to change if they are unhappy. 
+				int32 LayerIndex = 0;
+				for (; LayerIndex < ExistingPatchLayerIndex; ++LayerIndex)
+				{
+					TArray<ALandscapeBlueprintBrushBase*> LayerBrushes = InTargetLandscape->GetBrushesForLayer(LayerIndex);
+					if (LayerBrushes.ContainsByPredicate([](ALandscapeBlueprintBrushBase* Brush) 
+						{ 
+							// It would be preferable to do a Cast<ALandscapeBlueprintBrushBase> to be robust to subclassing, etc.,
+							// but that doesn't seem worth having to add a dependency on the water plugin, hence the comparison by name.
+							// Note that the "A" prefix is removed in GetName.
+							FString BrushManagerName = Brush->GetClass()->GetName();
+							return Brush->GetClass()->GetName() == TEXT("WaterBrushManager");
+						}))
+					{
+						// Found a water brush manager.
+						break;
+					}
+				}
+
+				if (LayerIndex < ExistingPatchLayerIndex)
+				{
+					InTargetLandscape->ReorderLayer(ExistingPatchLayerIndex, LayerIndex);
+					ExistingPatchLayerIndex = LayerIndex;
+				}
 			}
 			InTargetLandscape->AddBrushToLayer(ExistingPatchLayerIndex, this);
 		}
@@ -131,7 +176,7 @@ void ALandscapePatchManager::AddPatch(TObjectPtr<ULandscapePatchComponent> Patch
 	if (Patch)
 	{
 		Modify();
-		PatchComponents.AddUnique(Patch);
+		PatchComponents.AddUnique(TSoftObjectPtr<ULandscapePatchComponent>(Patch.Get()));
 		RequestLandscapeUpdate();
 	}
 }
@@ -143,7 +188,7 @@ bool ALandscapePatchManager::RemovePatch(TObjectPtr<ULandscapePatchComponent> Pa
 	if (Patch)
 	{
 		Modify();
-		bRemoved = PatchComponents.Remove(Patch) > 0;
+		bRemoved = PatchComponents.Remove(TSoftObjectPtr<ULandscapePatchComponent>(Patch.Get())) > 0;
 		RequestLandscapeUpdate();
 	}
 	
diff --git a/Engine/Plugins/Experimental/LandscapePatch/Source/LandscapePatch/Private/LandscapeTextureHeightPatch.cpp b/Engine/Plugins/Experimental/LandscapePatch/Source/LandscapePatch/Private/LandscapeTextureHeightPatch.cpp
--- a/Engine/Plugins/Experimental/LandscapePatch/Source/LandscapePatch/Private/LandscapeTextureHeightPatch.cpp
+++ b/Engine/Plugins/Experimental/LandscapePatch/Source/LandscapePatch/Private/LandscapeTextureHeightPatch.cpp
@@ -248,7 +248,7 @@ void ULandscapeTextureHeightPatch::PostEditChangeProperty(FPropertyChangedEvent&
 	{
 		if (PropertyChangedEvent.Property->GetFName() == GET_MEMBER_NAME_CHECKED(ULandscapeTextureHeightPatch, SourceMode))
 		{
-			ResetSourceMode(SourceMode);
+			SetSourceMode(SourceMode);
 			if (SourceMode == ELandscapeTexturePatchSourceMode::InternalTexture)
 			{
 				FTextureCompilingManager::Get().FinishCompilation({ InternalTexture });
@@ -387,10 +387,24 @@ void ULandscapeTextureHeightPatch::UpdateShaderParams(
 	Params.InZeroInEncoding = bNativeEncoding ? LandscapeDataAccess::MidValue : EncodingSettings.ZeroInEncoding;
 
 	Params.InHeightOffset = 0;
-	if (bUsePatchZAsReference)
+	switch (ZeroHeightMeaning)
+	{
+	case ELandscapeTextureHeightPatchZeroHeightMeaning::LandscapeZ:
+		break; // no offset necessary
+	case ELandscapeTextureHeightPatchZeroHeightMeaning::PatchZ:
 	{
 		FVector3d PatchOriginInHeightmapCoords = LandscapeHeightmapToWorld.InverseTransformPosition(PatchToWorld.GetTranslation());
 		Params.InHeightOffset = PatchOriginInHeightmapCoords.Z - LandscapeDataAccess::MidValue;
+		break;
+	}
+	case ELandscapeTextureHeightPatchZeroHeightMeaning::WorldZero:
+	{
+		FVector3d WorldOriginInHeightmapCoords = LandscapeHeightmapToWorld.InverseTransformPosition(FVector::ZeroVector);
+		Params.InHeightOffset = WorldOriginInHeightmapCoords.Z - LandscapeDataAccess::MidValue;
+		break;
+	}
+	default:
+		ensure(false);
 	}
 
 	// The outer half-pixel shouldn't affect the landscape because it is not part of our official coverage area.
@@ -403,6 +417,8 @@ void ULandscapeTextureHeightPatch::UpdateShaderParams(
 
 	Params.InFalloffWorldMargin = Falloff / FMath::Min(ComponentScale.X, ComponentScale.Y);
 
+	Params.InBlendMode = static_cast<uint32>(BlendMode);
+
 	// Pack our booleans into a bitfield
 	using EShaderFlags = FApplyLandscapeTextureHeightPatchPS::EFlags;
 	EShaderFlags Flags = EShaderFlags::None;
@@ -413,9 +429,6 @@ void ULandscapeTextureHeightPatch::UpdateShaderParams(
 	Flags |= bUseTextureAlphaChannel ?
 		EShaderFlags::ApplyPatchAlpha : EShaderFlags::None;
 
-	Flags |= (BlendMode == ELandscapeTextureHeightPatchBlendMode::Additive) ?
-		EShaderFlags::AdditiveMode : EShaderFlags::None;
-
 	Flags |= bNativeEncoding ?
 		EShaderFlags::InputIsPackedHeight : EShaderFlags::None;
 
@@ -626,12 +639,29 @@ void ULandscapeTextureHeightPatch::Reinitialize()
 		SourceEncoding = ELandscapeTextureHeightPatchEncoding::NativePackedHeight;
 		ResizeRenderTargetIfNeeded(InitTextureSizeX, InitTextureSizeY);
 			
-		// If bUsePatchZAsReference is true, then we're going to be adding an offset to our rendered landscape
-		// section in anticipation of adding an offset based on patch z later. In that case, we'll render
-		// the landscape to an intermediate target first, then do the copy to the internal target while adding
-		// in the offset. Otherwise, we can render directly to internal render target.
+		// If ZeroHeightMeaning is not landscape Z, then we're going to be applying an offset to our data when
+		// applying it to landscape, which means we'll need to apply the inverse offset when initializing here
+		// so that we get the same landscape back. In that case, we'll render the landscape to an intermediate
+		// target first, then do the copy to the internal target while adding in the offset. Otherwise, we can
+		// render directly to internal render target.
+		double OffsetToApply = 0;
+		if (ZeroHeightMeaning != ELandscapeTextureHeightPatchZeroHeightMeaning::LandscapeZ)
+		{
+			FTransform LandscapeHeightmapToWorld = PatchManager->GetHeightmapCoordsToWorld();
+			double ZeroHeight = 0;
+			if (ZeroHeightMeaning == ELandscapeTextureHeightPatchZeroHeightMeaning::PatchZ)
+			{
+				ZeroHeight = LandscapeHeightmapToWorld.InverseTransformPosition(GetComponentTransform().GetTranslation()).Z;
+			}
+			else if (ZeroHeightMeaning == ELandscapeTextureHeightPatchZeroHeightMeaning::WorldZero)
+			{
+				ZeroHeight = LandscapeHeightmapToWorld.InverseTransformPosition(FVector::ZeroVector).Z;
+			}
+			OffsetToApply = LandscapeDataAccess::MidValue - ZeroHeight;
+		}
+
 		UTextureRenderTarget2D* RenderedHeightmapSection = InternalRenderTarget;
-		if (bUsePatchZAsReference)
+		if (OffsetToApply != 0)
 		{
 			RenderedHeightmapSection = NewObject<UTextureRenderTarget2D>(this);
 			RenderedHeightmapSection->RenderTargetFormat = ETextureRenderTargetFormat::RTF_RG8;
@@ -646,14 +676,9 @@ void ULandscapeTextureHeightPatch::Reinitialize()
 
 		bApplyComponentZScale = false;
 
-		// ...However, we do need to apply the offset to account for our patch's position relative to the
-		// landscape z axis, if that is a setting that the user is using.
-		if (bUsePatchZAsReference)
+		// Apply the inverse offset we determined earlier.
+		if (OffsetToApply != 0)
 		{
-			FTransform LandscapeHeightmapToWorld = PatchManager->GetHeightmapCoordsToWorld();
-			FVector3d PatchOriginInHeightmapCoords = LandscapeHeightmapToWorld.InverseTransformPosition(GetComponentTransform().GetTranslation());
-			float OffsetToApply = LandscapeDataAccess::MidValue - PatchOriginInHeightmapCoords.Z;
-
 			ENQUEUE_RENDER_COMMAND(RenderHeightmap)([Input = RenderedHeightmapSection->GetResource(), Patch = InternalRenderTarget->GetResource(), OffsetToApply](FRHICommandListImmediate& RHICmdList)
 			{
 				FRDGBuilder GraphBuilder(RHICmdList, RDG_EVENT_NAME("OffsetTextureHeightPatchInitialization"));
@@ -828,7 +853,7 @@ bool ULandscapeTextureHeightPatch::ResizeTextureIfNeeded(int32 SizeX, int32 Size
 	return bChanged;
 }
 
-bool ULandscapeTextureHeightPatch::SetSourceMode(ELandscapeTexturePatchSourceMode NewMode, bool bDeleteUnusedInternalTextures)
+bool ULandscapeTextureHeightPatch::SetSourceMode(ELandscapeTexturePatchSourceMode NewMode, bool bInitializeIfRenderTarget)
 {
 	SourceMode = NewMode;
 
@@ -848,44 +873,36 @@ bool ULandscapeTextureHeightPatch::SetSourceMode(ELandscapeTexturePatchSourceMod
 		ConvertInternalRenderTargetBackFromNativeTexture();
 	}
 
-	if (bDeleteUnusedInternalTextures)
+	if (NewMode == ELandscapeTexturePatchSourceMode::TextureBackedRenderTarget 
+		&& !InternalRenderTarget && bInitializeIfRenderTarget)
 	{
-		if (NewMode != ELandscapeTexturePatchSourceMode::TextureBackedRenderTarget)
+		FVector2D DesiredResolution;
+		if (GetInitResolutionFromLandscape(ResolutionMultiplier, DesiredResolution))
 		{
-			InternalRenderTarget = nullptr;
+			ResizeRenderTargetIfNeeded(DesiredResolution.X, DesiredResolution.Y);
 		}
-		if (NewMode == ELandscapeTexturePatchSourceMode::TextureAsset)
+		else
 		{
-			InternalTexture = nullptr;
+			ResizeRenderTargetIfNeeded(InitTextureSizeX, InitTextureSizeY);
 		}
+		
 	}
 
-	PreviousSourceMode = NewMode;
-
-	return true;
-}
-#endif // WITH_EDITOR
-
-bool ULandscapeTextureHeightPatch::ResetSourceMode(ELandscapeTexturePatchSourceMode NewMode)
-{
-	if (NewMode == ELandscapeTexturePatchSourceMode::TextureAsset)
+	// Discard any unused internal textures
+	if (NewMode != ELandscapeTexturePatchSourceMode::TextureBackedRenderTarget)
 	{
-		ResetSourceEncodingMode(ELandscapeTextureHeightPatchEncoding::ZeroToOne);
+		InternalRenderTarget = nullptr;
 	}
-	else if (NewMode == ELandscapeTexturePatchSourceMode::TextureBackedRenderTarget)
+	if (NewMode == ELandscapeTexturePatchSourceMode::TextureAsset)
 	{
-		ResetSourceEncodingMode(ELandscapeTextureHeightPatchEncoding::WorldUnits);
+		InternalTexture = nullptr;
 	}
-	// The internal texture case ignores the encoding mode property so it does not need resetting. More
-	// importantly, not resetting it allows us to use the current encoding settings to convert properly
-	// into the native format inside SetSourceMode, which calls ConvertInternalRenderTargetToNativeTexture.
-	//else
-	//{
-	//	ResetSourceEncodingMode(ELandscapeTextureHeightPatchEncoding::NativePackedHeight);
-	//}
 
-	return SetSourceMode(NewMode, true);
+	PreviousSourceMode = NewMode;
+
+	return true;
 }
+#endif // WITH_EDITOR
 
 void ULandscapeTextureHeightPatch::SnapToLandscape()
 {
diff --git a/Engine/Plugins/Experimental/LandscapePatch/Source/LandscapePatch/Private/LandscapeTextureHeightPatchPS.cpp b/Engine/Plugins/Experimental/LandscapePatch/Source/LandscapePatch/Private/LandscapeTextureHeightPatchPS.cpp
--- a/Engine/Plugins/Experimental/LandscapePatch/Source/LandscapePatch/Private/LandscapeTextureHeightPatchPS.cpp
+++ b/Engine/Plugins/Experimental/LandscapePatch/Source/LandscapePatch/Private/LandscapeTextureHeightPatchPS.cpp
@@ -20,8 +20,12 @@ void FApplyLandscapeTextureHeightPatchPS::ModifyCompilationEnvironment(const FGl
 	// Make our flag choices match in the shader.
 	OutEnvironment.SetDefine(TEXT("RECTANGULAR_FALLOFF_FLAG"), static_cast<uint8>(FApplyLandscapeTextureHeightPatchPS::EFlags::RectangularFalloff));
 	OutEnvironment.SetDefine(TEXT("APPLY_PATCH_ALPHA_FLAG"), static_cast<uint8>(FApplyLandscapeTextureHeightPatchPS::EFlags::ApplyPatchAlpha));
-	OutEnvironment.SetDefine(TEXT("ADDITIVE_MODE_FLAG"), static_cast<uint8>(FApplyLandscapeTextureHeightPatchPS::EFlags::AdditiveMode));
 	OutEnvironment.SetDefine(TEXT("INPUT_IS_PACKED_HEIGHT_FLAG"), static_cast<uint8>(FApplyLandscapeTextureHeightPatchPS::EFlags::InputIsPackedHeight));
+
+	OutEnvironment.SetDefine(TEXT("ADDITIVE_MODE"), static_cast<uint8>(FApplyLandscapeTextureHeightPatchPS::EBlendMode::Additive));
+	OutEnvironment.SetDefine(TEXT("ALPHA_BLEND_MODE"), static_cast<uint8>(FApplyLandscapeTextureHeightPatchPS::EBlendMode::AlphaBlend));
+	OutEnvironment.SetDefine(TEXT("MIN_MODE"), static_cast<uint8>(FApplyLandscapeTextureHeightPatchPS::EBlendMode::Min));
+	OutEnvironment.SetDefine(TEXT("MAX_MODE"), static_cast<uint8>(FApplyLandscapeTextureHeightPatchPS::EBlendMode::Max));
 }
 
 void FApplyLandscapeTextureHeightPatchPS::AddToRenderGraph(FRDGBuilder& GraphBuilder, FParameters* InParameters, const FIntRect& DestinationBounds)
diff --git a/Engine/Plugins/Experimental/LandscapePatch/Source/LandscapePatch/Private/LandscapeTexturePatchBase.cpp b/Engine/Plugins/Experimental/LandscapePatch/Source/LandscapePatch/Private/LandscapeTexturePatchBase.cpp
--- a/Engine/Plugins/Experimental/LandscapePatch/Source/LandscapePatch/Private/LandscapeTexturePatchBase.cpp
+++ b/Engine/Plugins/Experimental/LandscapePatch/Source/LandscapePatch/Private/LandscapeTexturePatchBase.cpp
@@ -106,7 +106,7 @@ bool ULandscapeTexturePatchBase::GetInitResolutionFromLandscape(float Resolution
 		return false;
 	}
 
-	ResolutionOut = FVector2D::Zero();
+	ResolutionOut = FVector2D::One();
 
 	FVector LandscapeScale = Landscape->GetTransform().GetScale3D();
 	// We go off of the larger dimension so that our patch works in different rotations.
diff --git a/Engine/Plugins/Experimental/LandscapePatch/Source/LandscapePatch/Public/LandscapePatchComponent.h b/Engine/Plugins/Experimental/LandscapePatch/Source/LandscapePatch/Public/LandscapePatchComponent.h
--- a/Engine/Plugins/Experimental/LandscapePatch/Source/LandscapePatch/Public/LandscapePatchComponent.h
+++ b/Engine/Plugins/Experimental/LandscapePatch/Source/LandscapePatch/Public/LandscapePatchComponent.h
@@ -36,9 +36,28 @@ public:
 	UFUNCTION(BlueprintCallable, Category = "LandscapePatch")
 	void RequestLandscapeUpdate();
 
+	/**
+	 * Allows the patch to be disabled, so that it no longer affects the landscape. This can be useful
+	 * when deleting the patch is undesirable, usually when the disabling is temporary.
+	 */
+	UFUNCTION(BlueprintCallable, Category = "LandscapePatch")
+	void SetIsEnabled(bool bEnabledIn);
+
+	/**
+	 * @return false if the patch is marked as disabled and therefore can't affect the landscape.
+	 */
+	UFUNCTION(BlueprintCallable, Category = "LandscapePatch")
+	bool IsEnabled() { return bIsEnabled; }
+
 	UFUNCTION(BlueprintCallable, Category = "LandscapePatch")
 	FTransform GetLandscapeHeightmapCoordsToWorld() const;
 
+	UFUNCTION(BlueprintCallable, Category = "LandscapePatch")
+	virtual void SetLandscape(ALandscape* NewLandscape);
+
+	UFUNCTION(BlueprintCallable, Category = "LandscapePatch")
+	virtual void SetPatchManager(ALandscapePatchManager* NewPatchManager);
+
 	// For now we keep the patches largely editor-only, since we don't yet support runtime landscape editing.
 	// The above functions are also editor-only (and don't work at runtime), but can't be in WITH_EDITOR blocks
 	// so that they can be called from non-editor-only classes in editor contexts.
@@ -63,10 +82,10 @@ public:
 
 protected:
 	UPROPERTY(EditAnywhere, Category = Settings)
-	TWeakObjectPtr<ALandscape> Landscape = nullptr;
+	TSoftObjectPtr<ALandscape> Landscape = nullptr;
 
-	UPROPERTY()
-	TWeakObjectPtr<ALandscapePatchManager> PatchManager = nullptr;
+	UPROPERTY(EditAnywhere, Category = Settings, AdvancedDisplay)
+	TSoftObjectPtr<ALandscapePatchManager> PatchManager = nullptr;
 
 	// Determines whether the height patch was made by copying a different height patch.
 	bool bWasCopy = false;
@@ -75,9 +94,19 @@ protected:
 	// the first OnRegister call. It remains false from the first OnRegiter call onward, even
 	// if the component is unregistered.
 	bool bLoadedButNotYetRegistered = false;
+
+	/** 
+	 * When false, patch does not affect the landscape. Useful for temporarily disabling the patch. 
+	 */
+	UPROPERTY(EditAnywhere, Category = Settings)
+	bool bIsEnabled = true;
 private:
 	// Starts as false and gets set to true in construction, so gets used to set bWasCopy
 	// by checking the indicator value at the start of construction.
 	UPROPERTY()
 	bool bPropertiesCopiedIndicator = false;
+
+	// Used to properly transition to a different manager when editing it via the detail panel.
+	UPROPERTY()
+	TSoftObjectPtr<ALandscapePatchManager> PreviousPatchManager = nullptr;
 };
\ No newline at end of file
diff --git a/Engine/Plugins/Experimental/LandscapePatch/Source/LandscapePatch/Public/LandscapePatchManager.h b/Engine/Plugins/Experimental/LandscapePatch/Source/LandscapePatch/Public/LandscapePatchManager.h
--- a/Engine/Plugins/Experimental/LandscapePatch/Source/LandscapePatch/Public/LandscapePatchManager.h
+++ b/Engine/Plugins/Experimental/LandscapePatch/Source/LandscapePatch/Public/LandscapePatchManager.h
@@ -4,7 +4,7 @@
 
 #include "CoreMinimal.h"
 
-#include "LandscapeBlueprintBrushBase.h"
+#include "LandscapeBlueprintBrush.h"
 
 #include "LandscapePatchManager.generated.h"
 
@@ -19,7 +19,7 @@ class ULandscapePatchComponent;
 //~ we currently don't want to do because we think it will clutter the brush interface and may lose opportunities
 //~ for optimization... 
 UCLASS()
-class LANDSCAPEPATCH_API ALandscapePatchManager : public ALandscapeBlueprintBrushBase
+class LANDSCAPEPATCH_API ALandscapePatchManager : public ALandscapeBlueprintBrush
 {
 	GENERATED_BODY()
 
@@ -62,7 +62,7 @@ public:
 protected:
 	
 	UPROPERTY()
-	TArray<TWeakObjectPtr<ULandscapePatchComponent>> PatchComponents;
+	TArray<TSoftObjectPtr<ULandscapePatchComponent>> PatchComponents;
 
 	UPROPERTY()
 	FTransform HeightmapCoordsToWorld;
diff --git a/Engine/Plugins/Experimental/LandscapePatch/Source/LandscapePatch/Public/LandscapeTextureHeightPatch.h b/Engine/Plugins/Experimental/LandscapePatch/Source/LandscapePatch/Public/LandscapeTextureHeightPatch.h
--- a/Engine/Plugins/Experimental/LandscapePatch/Source/LandscapePatch/Public/LandscapeTextureHeightPatch.h
+++ b/Engine/Plugins/Experimental/LandscapePatch/Source/LandscapePatch/Public/LandscapeTextureHeightPatch.h
@@ -6,7 +6,7 @@
 
 #include "Engine/TextureRenderTarget2D.h"
 
-#include "LandscapeTextureHeightPatchPS.h" // FApplyLandscapeTextureHeightPatchPS::FParameters
+#include "LandscapeTextureHeightPatchPS.h" // FApplyLandscapeTextureHeightPatchPS::FParameters, FApplyLandscapeTextureHeightPatchPS::EBlendMode
 #include "LandscapeTexturePatchBase.h"
 #include "MatrixTypes.h"
 
@@ -36,12 +36,18 @@ enum class ELandscapeTextureHeightPatchBlendMode : uint8
 	// height using falloff/alpha. E.g. with no falloff and alpha 1, the landscape will
 	// be set directly to the height sampled from patch. With alpha 0.5, landscape height 
 	// will be averaged evenly with patch height.
-	AlphaBlend,
+	AlphaBlend = static_cast<uint8>(UE::Landscape::FApplyLandscapeTextureHeightPatchPS::EBlendMode::AlphaBlend),
 	
 	// Interpreting the landscape mid value as 0, use the texture patch as an offset to
 	// apply to the landscape. Falloff/alpha will just affect the degree to which the offset
 	// is applied (e.g. alpha of 0.5 will apply just half the offset).
-	Additive,
+	Additive = static_cast<uint8>(UE::Landscape::FApplyLandscapeTextureHeightPatchPS::EBlendMode::Additive),
+
+	// Like Alpha Blend mode, but limited to only lowering the existing landscape values
+	Min = static_cast<uint8>(UE::Landscape::FApplyLandscapeTextureHeightPatchPS::EBlendMode::Min),
+
+	// Like Alpha Blend mode, but limited to only raising the existing landscape values
+	Max = static_cast<uint8>(UE::Landscape::FApplyLandscapeTextureHeightPatchPS::EBlendMode::Max)
 };
 
 // Determines falloff method for the patch's influence.
@@ -77,6 +83,24 @@ enum class ELandscapeTextureHeightPatchEncoding : uint8
 	//~ appropriately when setting the encoding mode.
 };
 
+UENUM(BlueprintType)
+enum class ELandscapeTextureHeightPatchZeroHeightMeaning : uint8
+{
+	// Zero height corresponds to the patch vertical position relative to the landscape. This moves
+	// the results up and down as the patch moves up and down.
+	PatchZ,
+
+	// Zero height corresponds to Z = 0 in the local space of the landscape, regardless of the patch vertical
+	// position. For instance, if landscape transform has z=-100 in world, then writing height 0 will correspond
+	// to z=-100 in world coordinates, regardless of patch Z. 
+	LandscapeZ,
+
+	// Zero height corresponds to the height of the world origin relative to landscape. In other words, writing
+	// height 0 will correspond to world z = 0 regardless of patch Z or landscape transform (as long as landscape
+	// transform still has Z up in world coordinates).
+	WorldZero
+};
+
 //~ A struct in case we find that we need other encoding settings.
 USTRUCT(BlueprintType)
 struct LANDSCAPEPATCH_API FLandscapeTexturePatchEncodingSettings
@@ -118,7 +142,7 @@ public:
 	
 	// ULandscapeTexturePatchBase
 	virtual bool SetTextureResolution(FVector2D ResolutionIn) override;
-	virtual bool SetSourceMode(ELandscapeTexturePatchSourceMode NewMode, bool bDeleteUnusedInternalTextures = true) override;
+	virtual bool SetSourceMode(ELandscapeTexturePatchSourceMode NewMode, bool bInitializeIfRenderTarget = true) override;
 
 	// UActorComponent
 	virtual void OnComponentCreated() override;
@@ -132,12 +156,6 @@ public:
 	virtual void PostInitProperties() override;
 #endif
 
-	/**
-	 * Just like SetSourceMode, but resets encoding mode and settings to mode-specific defaults
-	 */
-	UFUNCTION(BlueprintCallable, Category = LandscapePatch)
-	bool ResetSourceMode(ELandscapeTexturePatchSourceMode NewMode);
-
 	/**
 	 * Deletes the internal render target and internal texture.
 	 */
@@ -213,13 +231,13 @@ public:
 	}
 
 	/**
-	 * Determine whether the patch Z relative to the landscape affects the results. See comment for bUsePatchZAsReference.
+	 * Set how zero height is interpreted, see comments in ELandscapeTextureHeightPatchZeroHeightMeaning.
 	 */
 	UFUNCTION(BlueprintCallable, Category = "LandscapePatch")
-	void SetUsePatchZAsReference(bool bUsePatchZIn) 
+	void SetZeroHeightMeaning(ELandscapeTextureHeightPatchZeroHeightMeaning ZeroHeightMeaningIn)
 	{ 
 		Modify();
-		bUsePatchZAsReference = bUsePatchZIn; 
+		ZeroHeightMeaning = ZeroHeightMeaningIn;
 	}
 
 	UFUNCTION(BlueprintCallable, Category = "LandscapePatch")
@@ -232,43 +250,38 @@ public:
 
 protected:
 
-	UPROPERTY(EditAnywhere, Category = Settings, AdvancedDisplay)
+	UPROPERTY(EditAnywhere, Category = Settings)
 	ELandscapeTextureHeightPatchBlendMode BlendMode = ELandscapeTextureHeightPatchBlendMode::AlphaBlend;
 	
 	/** When true, texture alpha channel will be used when applying the patch. */
-	UPROPERTY(EditAnywhere, Category = Settings, AdvancedDisplay)
-	bool bUseTextureAlphaChannel = false;
-
 	UPROPERTY(EditAnywhere, Category = Settings)
-	ELandscapeTextureHeightPatchFalloffMode FalloffMode = ELandscapeTextureHeightPatchFalloffMode::RoundedRectangle;
-
-	/**
-	 * Distance (in unscaled world coordinates) across which to smoothly fall off the patch effects.
-	 */
-	UPROPERTY(EditAnywhere, Category = Settings, meta = (ClampMin = "0", UIMax = "2000"))
-	float Falloff = 0;
+	bool bUseTextureAlphaChannel = false;
 
 	/**
-	 * Determine whether the patch Z relative to the landscape affects the results. When true, moving the
-	 * patch up and down will move the results up and down, which is useful for using the patch to move
-	 * landscape edits with a mesh. When false, the patch is always applied as if it is at z=0 relative
-	 * to landscape, which removes the need to worry about its vertical location and is useful when using
-	 * the patch in Blueprints.
-	 * Note that this setting also affects reinitialization from landscape, and it should probably match
-	 * usage when initializing.
+	 * How 0 height is interpreted.
 	 */
-	UPROPERTY(EditAnywhere, Category = Settings, AdvancedDisplay)
-	bool bUsePatchZAsReference = true;
+	UPROPERTY(EditAnywhere, Category = Settings)
+	ELandscapeTextureHeightPatchZeroHeightMeaning ZeroHeightMeaning = ELandscapeTextureHeightPatchZeroHeightMeaning::PatchZ;
 
-	/** How the values stored in the patch represent the height. */
-	UPROPERTY(EditAnywhere, Category = Settings, AdvancedDisplay, meta = (EditConditionHides,
+	/** How the values stored in the patch represent the height. Not customizable for Internal Texture source mode, which always uses native packed height. */
+	UPROPERTY(EditAnywhere, Category = Settings, meta = (
 		EditCondition = "SourceMode != ELandscapeTexturePatchSourceMode::InternalTexture"))
 	ELandscapeTextureHeightPatchEncoding SourceEncoding = ELandscapeTextureHeightPatchEncoding::NativePackedHeight;
 
-	UPROPERTY(EditAnywhere, Category = Settings, AdvancedDisplay, meta = (UIMin = "0", UIMax = "1", EditConditionHides,
+	/** Encoding settings. Not relevant when using native packed height as the encoding. */
+	UPROPERTY(EditAnywhere, Category = Settings, meta = (UIMin = "0", UIMax = "1",
 		EditCondition = "SourceMode != ELandscapeTexturePatchSourceMode::InternalTexture && SourceEncoding != ELandscapeTextureHeightPatchEncoding::NativePackedHeight"))
 	FLandscapeTexturePatchEncodingSettings EncodingSettings;
 
+	UPROPERTY(EditAnywhere, Category = Settings)
+	ELandscapeTextureHeightPatchFalloffMode FalloffMode = ELandscapeTextureHeightPatchFalloffMode::RoundedRectangle;
+
+	/**
+	 * Distance (in unscaled world coordinates) across which to smoothly fall off the patch effects.
+	 */
+	UPROPERTY(EditAnywhere, Category = Settings, meta = (ClampMin = "0", UIMax = "2000"))
+	float Falloff = 0;
+
 	/**
 	 * Whether to apply the patch Z scale to the height stored in the patch.
 	 */
@@ -279,7 +292,7 @@ protected:
 	 * Controls how the patch is initialized when invoking Reinitialize().
 	 */
 	UPROPERTY(EditAnywhere, Category = Initialization, meta = (
-		EditCondition = "SourceMode != ELandscapeTexturePatchSourceMode::TextureAsset", EditConditionHides))
+		EditCondition = "SourceMode != ELandscapeTexturePatchSourceMode::TextureAsset"))
 	ELandscapeTextureHeightPatchInitMode InitializationMode = ELandscapeTextureHeightPatchInitMode::FromLandscape;
 
 	/** When initializing from an texture asset, set the internal texture to have the same resolution. */
@@ -290,8 +303,7 @@ protected:
 
 	/** When initializing from landscape, set resolution based off of the landscape (and a multiplier). */
 	UPROPERTY(EditAnywhere, Category = Initialization, meta = (
-		EditCondition = "SourceMode != ELandscapeTexturePatchSourceMode::TextureAsset && (InitializationMode != ELandscapeTextureHeightPatchInitMode::TextureAsset || !bUseSameTextureDimensions)", 
-		EditConditionHides))
+		EditCondition = "SourceMode != ELandscapeTexturePatchSourceMode::TextureAsset && (InitializationMode != ELandscapeTextureHeightPatchInitMode::TextureAsset || !bUseSameTextureDimensions)"))
 	bool bBaseResolutionOffLandscape = true;
 
 	/** 
@@ -300,20 +312,19 @@ protected:
 	 * a value less that 1.0 will use lower.
 	 */
 	UPROPERTY(EditAnywhere, Category = Initialization, meta = (
-		EditCondition = "SourceMode != ELandscapeTexturePatchSourceMode::TextureAsset && (InitializationMode != ELandscapeTextureHeightPatchInitMode::TextureAsset || !bUseSameTextureDimensions) && bBaseResolutionOffLandscape", 
-		EditConditionHides))
+		EditCondition = "SourceMode != ELandscapeTexturePatchSourceMode::TextureAsset && (InitializationMode != ELandscapeTextureHeightPatchInitMode::TextureAsset || !bUseSameTextureDimensions) && bBaseResolutionOffLandscape"))
 	float ResolutionMultiplier = 1;
 
 	/** Texture width to use when reinitializing. */
 	UPROPERTY(EditAnywhere, Category = Initialization, meta = (
 		EditCondition = "SourceMode != ELandscapeTexturePatchSourceMode::TextureAsset && (InitializationMode != ELandscapeTextureHeightPatchInitMode::TextureAsset || !bUseSameTextureDimensions)", 
-		EditConditionHides, ClampMin = "1"))
+		ClampMin = "1"))
 	int32 InitTextureSizeX = 64;
 
 	/** Texture height to use when reinitializing */
 	UPROPERTY(EditAnywhere, Category = Initialization, meta = (
 		EditCondition = "SourceMode != ELandscapeTexturePatchSourceMode::TextureAsset && (InitializationMode != ELandscapeTextureHeightPatchInitMode::TextureAsset || !bUseSameTextureDimensions)", 
-		EditConditionHides, ClampMin = "1"))
+		ClampMin = "1"))
 	int32 InitTextureSizeY = 64;
 
 	// Used to properly transition the source mode when editing it via the detail panel.
diff --git a/Engine/Plugins/Experimental/LandscapePatch/Source/LandscapePatch/Public/LandscapeTextureHeightPatchPS.h b/Engine/Plugins/Experimental/LandscapePatch/Source/LandscapePatch/Public/LandscapeTextureHeightPatchPS.h
--- a/Engine/Plugins/Experimental/LandscapePatch/Source/LandscapePatch/Public/LandscapeTextureHeightPatchPS.h
+++ b/Engine/Plugins/Experimental/LandscapePatch/Source/LandscapePatch/Public/LandscapeTextureHeightPatchPS.h
@@ -20,26 +20,41 @@ class LANDSCAPEPATCH_API FApplyLandscapeTextureHeightPatchPS : public FGlobalSha
 
 public:
 
+	enum class EBlendMode : uint8
+	{
+		/** Desired height is alpha blended with the current. */
+		AlphaBlend,
+
+		/** Desired height is multiplied by alpha and added to current. */
+		Additive,
+
+		/** Like AlphaBlend, but patch is limited to only lowering the landscape. */
+		Min,
+
+		/** Like AlphaBlend, but patch is limited to only raising the landscape. */
+		Max
+	};
+
 	// Flags that get packed into a bitfield because we're not allowed to use bool shader parameters:
 	enum class EFlags : uint8
 	{
 		None = 0,
 
+		// When false, falloff is circular.
 		RectangularFalloff = 1 << 0,
+
+		// When true, the texture alpha channel is considered for blending (in addition to falloff, if nonzero)
 		ApplyPatchAlpha = 1 << 1,
-		AdditiveMode = 1 << 2,
-		InputIsPackedHeight = 1 << 3
+
+		// When false, the input is directly interpreted as being the height value to process. When true, the height
+		// is unpacked from the red and green channels to make a 16 bit int.
+		InputIsPackedHeight = 1 << 2
 	};
 
-	// When false, falloff is circular.
-	static const uint32 RectangularFalloffFlag = 1 << 0;
-	// When true, the texture alpha channel is considered for blending (in addition to falloff, if nonzero)
-	static const uint32 ApplyPatchAlphaFlag = 1 << 1;
-	// When true, the patch is added to the landscape rather than alpha blending with the landscape.
-	static const uint32 AdditiveModeFlag = 1 << 2;
-	// When false, the input is directly interpreted as being the height value to process. When true, the height
-	// is unpacked from the red and green channels to make a 16 bit int.
-	static const uint32 InputIsPackedHeightFlag = 1 << 3;
+	// TODO: We could consider exposing an additional global alpha setting that we can use to pass in the given
+	// edit layer alpha value... On the other hand, we currently don't bother doing this in any existing blueprint
+	// brushes, and it would be hard to support in a way that doesn't require each blueprint brush to respect it
+	// individually... Not clear whether this is something worth doing yet.
 
 	BEGIN_SHADER_PARAMETER_STRUCT(FParameters, )
 		SHADER_PARAMETER_RDG_TEXTURE_SRV(Texture2D<float4>, InSourceHeightmap)
@@ -58,6 +73,7 @@ public:
 		// In patch texture space, the size of the margin across which the alpha falls from 1 to 0
 		SHADER_PARAMETER(float, InFalloffWorldMargin)
 		SHADER_PARAMETER(FVector2f, InPatchWorldDimensions)
+		SHADER_PARAMETER(uint32, InBlendMode)
 		// Some combination of the flags (see constants above).
 		SHADER_PARAMETER(uint32, InFlags)
 
diff --git a/Engine/Plugins/Experimental/LandscapePatch/Source/LandscapePatch/Public/LandscapeTexturePatchBase.h b/Engine/Plugins/Experimental/LandscapePatch/Source/LandscapePatch/Public/LandscapeTexturePatchBase.h
--- a/Engine/Plugins/Experimental/LandscapePatch/Source/LandscapePatch/Public/LandscapeTexturePatchBase.h
+++ b/Engine/Plugins/Experimental/LandscapePatch/Source/LandscapePatch/Public/LandscapeTexturePatchBase.h
@@ -56,8 +56,11 @@ public:
 	UFUNCTION(BlueprintCallable, Category = "LandscapePatch")
 	UTextureRenderTarget2D* GetInternalRenderTarget() { return InternalRenderTarget;}
 
-	// These need to be public so that we can take the internal texture and write it to an external one
+	// This needs to be public so that we can take the internal texture and write it to an external one,
+	// but unclear whether we want to expose it to blueprints, since it's a fairly internal thing.
 	UTexture2D* GetInternalTexture() { return InternalTexture; }
+
+	UFUNCTION(BlueprintCallable, Category = "LandscapePatch")
 	void SetTextureAsset(UTexture* TextureIn);
 
 	virtual ELandscapeTexturePatchSourceMode GetSourceMode() const { return SourceMode; }
@@ -68,7 +71,7 @@ public:
 	 * @return true If successful.
 	 */
 	UFUNCTION(BlueprintCallable, Category = "LandscapePatch")
-	virtual bool SetSourceMode(ELandscapeTexturePatchSourceMode NewMode, bool bDiscardUnusedInternalTextures = true) { return false; }
+	virtual bool SetSourceMode(ELandscapeTexturePatchSourceMode NewMode, bool bInitializeIfRenderTarget = true) { return false; }
 
 	/**
 	 * Gets the transform from patch to world. The transform is based off of the component
