LandscapePatch: Fix 20778728 breaking CIS runtime build via LandscapeEditorUtilities.

diff --git a/Engine/Plugins/Experimental/LandscapePatch/Source/LandscapePatch/LandscapePatch.Build.cs b/Engine/Plugins/Experimental/LandscapePatch/Source/LandscapePatch/LandscapePatch.Build.cs
--- a/Engine/Plugins/Experimental/LandscapePatch/Source/LandscapePatch/LandscapePatch.Build.cs
+++ b/Engine/Plugins/Experimental/LandscapePatch/Source/LandscapePatch/LandscapePatch.Build.cs
@@ -39,11 +39,6 @@ public class LandscapePatch : ModuleRules
 				"RenderCore",
 				"RHI",
 				"Projects", // IPluginManager
-				
-				// This doesn't seem to be editor-only currently. If it becomes editor-only, may need
-				// to inherit from LanscapeBlueprintBrushBase instead, and will lose the ability to
-				// add the patch manager via the Blueprints tool in Landscape mode.
-				"LandscapeEditorUtilities",
 			}
 			);
 	}
diff --git a/Engine/Plugins/Experimental/LandscapePatch/Source/LandscapePatch/Private/LandscapePatchManager.cpp b/Engine/Plugins/Experimental/LandscapePatch/Source/LandscapePatch/Private/LandscapePatchManager.cpp
--- a/Engine/Plugins/Experimental/LandscapePatch/Source/LandscapePatch/Private/LandscapePatchManager.cpp
+++ b/Engine/Plugins/Experimental/LandscapePatch/Source/LandscapePatch/Private/LandscapePatchManager.cpp
@@ -13,7 +13,7 @@
 
 // TODO: Not sure if using this kind of constructor is a proper thing to do vs some other hook...
 ALandscapePatchManager::ALandscapePatchManager(const FObjectInitializer& ObjectInitializer)
-	: ALandscapeBlueprintBrush(ObjectInitializer)
+	: ALandscapeBlueprintBrushBase(ObjectInitializer)
 {
 #if WITH_EDITOR
 	SetAffectsHeightmap(true);
diff --git a/Engine/Plugins/Experimental/LandscapePatch/Source/LandscapePatch/Public/LandscapePatchManager.h b/Engine/Plugins/Experimental/LandscapePatch/Source/LandscapePatch/Public/LandscapePatchManager.h
--- a/Engine/Plugins/Experimental/LandscapePatch/Source/LandscapePatch/Public/LandscapePatchManager.h
+++ b/Engine/Plugins/Experimental/LandscapePatch/Source/LandscapePatch/Public/LandscapePatchManager.h
@@ -4,7 +4,7 @@
 
 #include "CoreMinimal.h"
 
-#include "LandscapeBlueprintBrush.h"
+#include "LandscapeBlueprintBrushBase.h"
 
 #include "LandscapePatchManager.generated.h"
 
@@ -19,7 +19,7 @@ class ULandscapePatchComponent;
 //~ we currently don't want to do because we think it will clutter the brush interface and may lose opportunities
 //~ for optimization... 
 UCLASS()
-class LANDSCAPEPATCH_API ALandscapePatchManager : public ALandscapeBlueprintBrush
+class LANDSCAPEPATCH_API ALandscapePatchManager : public ALandscapeBlueprintBrushBase
 {
 	GENERATED_BODY()
 
