LandscapePatch: Fix lots of managers being created if landscape is set in a BP actor component in BP editor.

diff --git a/Engine/Plugins/Experimental/LandscapePatch/Source/LandscapePatch/Private/LandscapePatchComponent.cpp b/Engine/Plugins/Experimental/LandscapePatch/Source/LandscapePatch/Private/LandscapePatchComponent.cpp
--- a/Engine/Plugins/Experimental/LandscapePatch/Source/LandscapePatch/Private/LandscapePatchComponent.cpp
+++ b/Engine/Plugins/Experimental/LandscapePatch/Source/LandscapePatch/Private/LandscapePatchComponent.cpp
@@ -14,6 +14,22 @@ namespace LandscapePatchComponentLocals
 {
 #if WITH_EDITOR
 
+	ALandscapePatchManager* FindExistingPatchManagerForLandscape(ALandscape* Landscape)
+	{
+		for (int LayerIndex = 0; LayerIndex < Landscape->LandscapeLayers.Num(); ++LayerIndex)
+		{
+			TArray<ALandscapeBlueprintBrushBase*> LayerBrushes = Landscape->GetBrushesForLayer(LayerIndex);
+			ALandscapeBlueprintBrushBase** Found = LayerBrushes.FindByPredicate(
+				[](ALandscapeBlueprintBrushBase* Candidate) { return Cast<ALandscapePatchManager>(Candidate) != nullptr; });
+			if (Found)
+			{
+				return Cast<ALandscapePatchManager>(*Found);
+			}
+		}
+
+		return nullptr;
+	}
+
 	ALandscapePatchManager* CreateNewPatchManagerForLandscape(ALandscape* Landscape)
 	{
 		if (!ensure(Landscape->CanHaveLayersContent()))
@@ -76,6 +92,13 @@ void ULandscapePatchComponent::OnComponentCreated()
 	UWorld* World = GetWorld();
 	AActor* OwningActor = GetOwner();
 
+	if (HasAnyFlags(RF_ClassDefaultObject | RF_ArchetypeObject))
+	{
+		// If this is a template object from the blueprint editor or something of the sort, don't do
+		// anything else special- we don't want to add to patch manager, etc.
+		return;
+	}
+
 	if (PatchManager.IsValid())
 	{
 		// If we copied over a patch manager, presumably Landscape should be
@@ -86,10 +109,16 @@ void ULandscapePatchComponent::OnComponentCreated()
 	{
 		if (Landscape.IsValid())
 		{
-			// If we copied over a patch with a landscape but no manager, create manager in that landscape
+			// If we copied over a patch with a landscape but no manager, find or create a manager in that landscape
 			if (Landscape->CanHaveLayersContent())
 			{
-				SetPatchManager(CreateNewPatchManagerForLandscape(Landscape.Get()));
+				ALandscapePatchManager* ManagerToUse = FindExistingPatchManagerForLandscape(Landscape.Get());
+				if (!ManagerToUse)
+				{
+					ManagerToUse = CreateNewPatchManagerForLandscape(Landscape.Get());
+				}
+
+				SetPatchManager(ManagerToUse);
 			}
 			else
 			{
@@ -220,8 +249,8 @@ void ULandscapePatchComponent::SetLandscape(ALandscape* NewLandscape)
 	// to be able to use this function to reorder patches. On the other hand, it seems inconvenient
 	// to accidentally swap patch managers if there are multiple in the same landscape, and we kept
 	// landscape the same. It's hard to know the ideal behavior, but for now we'll keep it.
-	if (Landscape == NewLandscape && ((!Landscape && !PatchManager)
-		|| (Landscape && PatchManager && PatchManager->GetOwningLandscape() == Landscape)))
+	if (Landscape.Get() == NewLandscape && ((Landscape.IsNull() && !PatchManager)
+		|| (Landscape.IsValid() && PatchManager && PatchManager->GetOwningLandscape() == Landscape.Get())))
 	{
 		return;
 	}
@@ -235,34 +264,21 @@ void ULandscapePatchComponent::SetLandscape(ALandscape* NewLandscape)
 	}
 
 	// If landscape was valid, try to find a patch manager inside that landscape.
-	bool bFoundExisting = false;
-	for (int32 LayerIndex = 0; LayerIndex < Landscape->Layers.Num(); ++LayerIndex)
+	if (Landscape->CanHaveLayersContent())
 	{
-		TArray<ALandscapeBlueprintBrushBase*> Brushes = Landscape->GetBrushesForLayer(LayerIndex);
-		ALandscapeBlueprintBrushBase** ExistingManager = Brushes.FindByPredicate(
-			[](ALandscapeBlueprintBrushBase* Brush) { return Cast<ALandscapePatchManager>(Brush) != nullptr; });
-		if (ExistingManager)
+		ALandscapePatchManager* ManagerToUse = FindExistingPatchManagerForLandscape(Landscape.Get());
+		if (!ManagerToUse)
 		{
-			bFoundExisting = true;
-			SetPatchManager(Cast<ALandscapePatchManager>(*ExistingManager));
-			break;
+			ManagerToUse = CreateNewPatchManagerForLandscape(Landscape.Get());
 		}
-	}
 
-	// If we didn't find an existing manager, create one for this landscape.
-	if (!bFoundExisting)
+		SetPatchManager(ManagerToUse);
+	}
+	else
 	{
-		if (Landscape->CanHaveLayersContent())
-		{
-			SetPatchManager(CreateNewPatchManagerForLandscape(Landscape.Get()));
-			PatchManager->AddPatch(this);
-		}
-		else
-		{
-			UE_LOG(LogLandscapePatch, Warning, TEXT("Landscape target for height patch did not have edit layers enabled. Unable to create patch manager."));
-			Landscape = nullptr;
-			SetPatchManager(nullptr);
-		}
+		UE_LOG(LogLandscapePatch, Warning, TEXT("Landscape target for height patch did not have edit layers enabled. Unable to create patch manager."));
+		Landscape = nullptr;
+		SetPatchManager(nullptr);
 	}
 
 #endif // WITH_EDITOR
@@ -287,7 +303,10 @@ void ULandscapePatchComponent::SetPatchManager(ALandscapePatchManager* NewPatchM
 	PatchManager = NewPatchManager;
 	if (NewPatchManager)
 	{
-		PatchManager->AddPatch(this);
+		if (!HasAnyFlags(RF_ClassDefaultObject | RF_ArchetypeObject))
+		{
+			PatchManager->AddPatch(this);
+		}
 		Landscape = PatchManager->GetOwningLandscape();
 	}
 	else
