LandscapePatch: Create LandscapeTexturePatch, which combines a height patch with weight patches.

diff --git a/Engine/Plugins/Experimental/LandscapePatch/Shaders/Private/LandscapeTextureHeightPatchPS.usf b/Engine/Plugins/Experimental/LandscapePatch/Shaders/Private/LandscapeTextureHeightPatchPS.usf
deleted file mode 100644
diff --git a/Engine/Plugins/Experimental/LandscapePatch/Shaders/Private/LandscapeTexturePatchPS.usf b/Engine/Plugins/Experimental/LandscapePatch/Shaders/Private/LandscapeTexturePatchPS.usf
new file mode 100644
--- /dev/null
+++ b/Engine/Plugins/Experimental/LandscapePatch/Shaders/Private/LandscapeTexturePatchPS.usf
@@ -0,0 +1,286 @@
+// Copyright Epic Games, Inc. All Rights Reserved.
+
+#pragma once
+
+#include "/Engine/Private/Common.ush"
+#include "/Engine/Private/Landscape/LandscapeCommon.ush"
+
+#if defined (__INTELLISENSE__)
+// Uncomment the appropriate define for enabling syntax highlighting with HLSL Tools for Visual Studio : 
+//#define APPLY_HEIGHT_PATCH 1
+//#define OFFSET_HEIGHT_PATCH 1
+//#define SIMPLE_TEXTURE_COPY 1
+//#define CONVERT_TO_NATIVE_LANDSCAPE_PATCH 1
+//#define CONVERT_BACK_FROM_NATIVE_LANDSCAPE_PATCH 1
+//#define APPLY_WEIGHT_PATCH 1
+//#define REINITIALIZE_PATCH 1
+#endif // defined (__INTELLISENSE__)
+
+// This constant gets used in multiple shaders, hence it's out here.
+static const float LANDSCAPE_MID_VALUE = 32768.0f;
+
+#if APPLY_HEIGHT_PATCH || APPLY_WEIGHT_PATCH
+
+// Given the falloff settings, gives the alpha to use for a pixel at the given patch UV coordinates.
+float GetFalloffAlpha(float FalloffWorldMargin, float2 PatchWorldDimensions, float2 PatchUVCoordinates, float2 EdgeUVDeadBorder, bool bRectangularFalloff)
+{
+	static const float HALF_PI = 3.14159265f / 2;
+	static const float FLOAT_TOLERANCE = 0.0001;
+	
+	// The falloff starts at the edges and goes inward, and shouldn't be higher than 1/2 of the smaller dimension, or
+	// else it extends past the center. If it does extend past the center, we'll decrease the max alpha in the center
+	// (imagine shaving off a cube with progressively flatter slopes- once you get past the center of the top, you get
+	// a progressively shorter pyramid). Maybe we shouldn't even bother dealing with this case, but it's easy enough to do.
+	float ClampedFalloffMargin = clamp(FalloffWorldMargin, 0, min(PatchWorldDimensions.x, PatchWorldDimensions.y) / 2);
+	float bFalloffAlpha = FalloffWorldMargin > 0 ? ClampedFalloffMargin / FalloffWorldMargin : 1;
+	
+	// Alpha T is a parameter that we'll pass through a cos^2 function to get the alpha (to have it be smoother). It should be 
+	// 0 at the places where there is no falloff (ie alpha is 1), and go to 1 at places where falloff is complete (and falloff is 0)
+	float AlphaT = 1;
+	
+	if (bRectangularFalloff)
+	{
+		float2 DistancesFromEdges = PatchWorldDimensions * min(PatchUVCoordinates - EdgeUVDeadBorder, 1 - EdgeUVDeadBorder - PatchUVCoordinates);
+		bool bIsInsidePatch = all(DistancesFromEdges >= -FLOAT_TOLERANCE);
+		AlphaT = bIsInsidePatch ? 0 : 1;
+		
+		if (bIsInsidePatch && ClampedFalloffMargin > 0)
+		{
+			// Interpret distances as proportions of falloff margin
+			DistancesFromEdges = DistancesFromEdges / ClampedFalloffMargin;
+		
+			if (DistancesFromEdges.x > 0 && DistancesFromEdges.x < 1.0
+				&& DistancesFromEdges.y > 0 && DistancesFromEdges.y < 1.0)
+			{
+				AlphaT = min(1.0, length(float2(1, 1) - DistancesFromEdges));
+			}
+			else
+			{
+				AlphaT = 1 - min(1.0, min(DistancesFromEdges.x, DistancesFromEdges.y));
+			}
+		}
+	}
+	else
+	{
+		float OuterRadius = min(PatchWorldDimensions.x, PatchWorldDimensions.y) / 2;
+		float DistanceFromCenter = length(PatchWorldDimensions * (PatchUVCoordinates - float2(0.5, 0.5)));
+		float DistanceFromOuterCircle = OuterRadius - DistanceFromCenter;
+		
+		bool bIsInsideCircle = DistanceFromOuterCircle > 0;
+		AlphaT = bIsInsideCircle ? 0 : 1;
+		if (bIsInsideCircle && ClampedFalloffMargin > 0)
+		{
+			AlphaT = 1 - DistanceFromOuterCircle / ClampedFalloffMargin;
+			AlphaT = max(0, AlphaT); // clamp to 0 so we don't go past ClampedFalloffMargin
+		}
+	}
+	
+	float Alpha = cos(AlphaT * HALF_PI);
+	Alpha = Alpha * Alpha * bFalloffAlpha;
+	
+	return Alpha;
+}
+#endif // APPLY_HEIGHT_PATCH || APPLY_WEIGHT_PATCH
+
+#if APPLY_HEIGHT_PATCH
+Texture2D<float4> InSourceHeightmap;
+Texture2D<float4> InHeightPatch;
+SamplerState InHeightPatchSampler;
+float4x4 InHeightmapToPatch;
+float2 InPatchWorldDimensions;
+float2 InEdgeUVDeadBorder;
+float InFalloffWorldMargin;
+// Value that corresponds to 0 (i.e. to LANDSCAPE_MID_VALUE) in the patch data
+float InZeroInEncoding;
+// Scale to apply to patch values relative to the value representing zero
+float InHeightScale;
+// An offset to apply to the resulting height, after applying the height scale
+float InHeightOffset;
+// Enum defining blend modes, with values set in corresponding cpp file.
+uint InBlendMode;
+// A combination of flags, whose positions are set in the corresponding cpp file.
+uint InFlags;
+
+void ApplyLandscapeTextureHeightPatch(in float4 SVPos : SV_POSITION, out float2 OutColor : SV_Target0)
+{
+	// Unpack our flags. The flag choices are set in FApplyLandscapeTextureHeightPatchPS::ModifyCompilationEnvironment
+	// so that they match to the cpp file
+	bool bRectangularFalloff = InFlags & RECTANGULAR_FALLOFF_FLAG;
+	bool bApplyPatchAlpha = InFlags & APPLY_PATCH_ALPHA_FLAG;
+	bool bInputIsPackedHeight = InFlags & INPUT_IS_PACKED_HEIGHT_FLAG;
+	
+	// We need only the 2D affine transformation that goes from landscape XY heightmap integer coordinates
+	// to patch UV coordinates.
+	float2x2 HeightmapToPatchRotateScale = (float2x2) InHeightmapToPatch;
+	float2 HeightmapToPatchTranslate = InHeightmapToPatch._m03_m13;
+	
+	float2 PatchUVCoordinates = mul(HeightmapToPatchRotateScale, SVPos.xy) + HeightmapToPatchTranslate;
+	float4 PatchSampledValue = InHeightPatch.Sample(InHeightPatchSampler, PatchUVCoordinates);
+	float PatchStoredHeight = bInputIsPackedHeight ? UnpackHeight(PatchSampledValue.xy) : PatchSampledValue.x;
+	
+	float PatchSignedHeight = InHeightScale * (PatchStoredHeight - InZeroInEncoding) + InHeightOffset;
+	
+	int2 HeightmapCoordinates = floor(SVPos.xy);
+	float4 CurrentPackedHeight = InSourceHeightmap.Load(int3(HeightmapCoordinates, 0));
+	float CurrentHeight = UnpackHeight(CurrentPackedHeight.xy);
+	
+	float Alpha = GetFalloffAlpha(InFalloffWorldMargin, InPatchWorldDimensions, PatchUVCoordinates, InEdgeUVDeadBorder, bRectangularFalloff);
+	
+	if (bApplyPatchAlpha)
+	{
+		Alpha *= PatchSampledValue.a;
+	}
+	
+	float NewHeight = 0;
+	switch (InBlendMode)
+	{
+		case ADDITIVE_MODE:
+			NewHeight = CurrentHeight + Alpha * PatchSignedHeight;
+			break;
+		case ALPHA_BLEND_MODE:
+			NewHeight = lerp(CurrentHeight, LANDSCAPE_MID_VALUE + PatchSignedHeight, Alpha);
+			break;
+		case MIN_MODE:
+			NewHeight = lerp(CurrentHeight, min(CurrentHeight, LANDSCAPE_MID_VALUE + PatchSignedHeight), Alpha);
+			break;
+		case MAX_MODE:
+			NewHeight = lerp(CurrentHeight, max(CurrentHeight, LANDSCAPE_MID_VALUE + PatchSignedHeight), Alpha);
+			break;
+	}
+	
+	OutColor = PackHeight(NewHeight);
+}
+#endif // APPLY_HEIGHT_PATCH
+
+#if OFFSET_HEIGHT_PATCH
+Texture2D<float4> InHeightmap;
+float InHeightOffset;
+
+void ApplyOffsetToHeightmap(in float4 SVPos : SV_POSITION, out float2 OutColor : SV_Target0)
+{
+	float StoredHeight = UnpackHeight(InHeightmap.Load(int3(SVPos.xy, 0)).xy);
+	OutColor = PackHeight(StoredHeight + InHeightOffset);
+}
+#endif // OFFSET_HEIGHT_PATCH
+
+#if CONVERT_TO_NATIVE_LANDSCAPE_PATCH
+Texture2D<float4> InHeightmap;
+float InZeroInEncoding;
+float InHeightScale;
+float InHeightOffset;
+
+void ConvertToNativeLandscapePatch(in float4 SVPos : SV_POSITION, out float2 OutColor : SV_Target0)
+{
+	float StoredHeight = InHeightmap.Load(int3(SVPos.xy, 0)).x;
+	// Apply the same kinds of transformations that we would apply for alpha blend mode in ApplyLandscapeTextureHeightPatch
+	OutColor = PackHeight(InHeightScale * (StoredHeight - InZeroInEncoding) + InHeightOffset + LANDSCAPE_MID_VALUE);
+}
+#endif // CONVERT_TO_NATIVE_LANDSCAPE_PATCH
+
+#if CONVERT_BACK_FROM_NATIVE_LANDSCAPE_PATCH
+Texture2D<float4> InHeightmap;
+float InZeroInEncoding;
+float InHeightScale;
+float InHeightOffset;
+
+void ConvertBackFromNativeLandscapePatch(in float4 SVPos : SV_POSITION, out float OutColor : SV_Target0)
+{
+	float StoredHeight = UnpackHeight(InHeightmap.Load(int3(SVPos.xy, 0)).xy);
+	// Undo the transformation we do in ConvertToNativeLandscapePatch
+	OutColor = InHeightScale ? (StoredHeight - LANDSCAPE_MID_VALUE - InHeightOffset) / InHeightScale + InZeroInEncoding
+		: InZeroInEncoding;
+}
+#endif // CONVERT_BACK_FROM_NATIVE_LANDSCAPE_PATCH
+
+#if APPLY_WEIGHT_PATCH
+Texture2D<float4> InSourceWeightmap;
+Texture2D<float4> InWeightPatch;
+SamplerState InWeightPatchSampler;
+float4x4 InWeightmapToPatch;
+float2 InPatchWorldDimensions;
+float2 InEdgeUVDeadBorder;
+float InFalloffWorldMargin;
+// Enum defining blend modes, with values set in corresponding cpp file.
+uint InBlendMode;
+// A combination of flags, whose positions are set in the corresponding cpp file.
+uint InFlags;
+
+
+void ApplyLandscapeTextureWeightPatch(in float4 SVPos : SV_POSITION, out float OutColor : SV_Target0)
+{
+	// Unpack our flags. The flag choices are set in FApplyLandscapeTextureWeightPatchPS::ModifyCompilationEnvironment
+	// so that they match to the cpp file
+	bool bRectangularFalloff = InFlags & RECTANGULAR_FALLOFF_FLAG;
+	bool bApplyPatchAlpha = InFlags & APPLY_PATCH_ALPHA_FLAG;
+	
+	// We need only the 2D affine transformation that goes from landscape weightmap integer coordinates
+	// to patch UV coordinates.
+	float2x2 WeightmapToPatchRotateScale = (float2x2) InWeightmapToPatch;
+	float2 WeightmapToPatchTranslate = InWeightmapToPatch._m03_m13;
+	
+	float2 PatchUVCoordinates = mul(WeightmapToPatchRotateScale, SVPos.xy) + WeightmapToPatchTranslate;
+	float4 PatchSampledValue = InWeightPatch.Sample(InWeightPatchSampler, PatchUVCoordinates);
+	float PatchWeight = PatchSampledValue.x;
+	
+	int2 WeightmapCoordinates = floor(SVPos.xy);
+	float4 CurrentWeight = InSourceWeightmap.Load(int3(WeightmapCoordinates, 0)).x;
+	
+	float Alpha = GetFalloffAlpha(InFalloffWorldMargin, InPatchWorldDimensions, PatchUVCoordinates, InEdgeUVDeadBorder, bRectangularFalloff);
+	
+	if (bApplyPatchAlpha)
+	{
+		Alpha *= PatchSampledValue.a;
+	}
+	
+	float NewWeight = 0;
+	switch (InBlendMode)
+	{
+		case ADDITIVE_MODE:
+			NewWeight = clamp(CurrentWeight + Alpha * PatchWeight, 0, 1);
+			break;
+		case ALPHA_BLEND_MODE:
+			NewWeight = lerp(CurrentWeight, PatchWeight, Alpha);
+			break;
+		case MIN_MODE:
+			NewWeight = lerp(CurrentWeight, min(CurrentWeight, PatchWeight), Alpha);
+			break;
+		case MAX_MODE:
+			NewWeight = lerp(CurrentWeight, max(CurrentWeight, PatchWeight), Alpha);
+			break;
+	}
+	
+	OutColor = NewWeight;
+}
+
+
+#endif // APPLY_WEIGHT_PATCH
+
+#if REINITIALIZE_PATCH
+Texture2D<float4> InSource;
+SamplerState InSourceSampler;
+float4x4 InPatchToSource;
+
+void ReinitializePatch(in float4 SVPos : SV_POSITION, out float4 OutColor : SV_Target0)
+{
+	// We need only the 2D affine transformation that goes from patch integer coordinates
+	// to source UV coordinates.
+	float2x2 PatchToSourceRotateScale = (float2x2) InPatchToSource;
+	float2 PatchToSourceTranslate = InPatchToSource._m03_m13;
+	
+	float2 SourceUVCoordinates = mul(PatchToSourceRotateScale, SVPos.xy) + PatchToSourceTranslate;
+	OutColor = InSource.Sample(InSourceSampler, SourceUVCoordinates);
+}
+
+#endif // REINITIALIZE_PATCH
+
+#if SIMPLE_TEXTURE_COPY
+Texture2D<float4> InSource;
+SamplerState InSourceSampler;
+float2 InDestinationResolution;
+
+void SimpleTextureCopy(in float4 SVPos : SV_POSITION, out float4 OutColor : SV_Target0)
+{
+	float2 UV = SVPos.xy / InDestinationResolution;
+	OutColor = InSource.Sample(InSourceSampler, UV);
+}
+#endif // SIMPLE_TEXTURE_COPY
\ No newline at end of file
diff --git a/Engine/Plugins/Experimental/LandscapePatch/Source/LandscapePatch/Private/LandscapePatchManager.cpp b/Engine/Plugins/Experimental/LandscapePatch/Source/LandscapePatch/Private/LandscapePatchManager.cpp
--- a/Engine/Plugins/Experimental/LandscapePatch/Source/LandscapePatch/Private/LandscapePatchManager.cpp
+++ b/Engine/Plugins/Experimental/LandscapePatch/Source/LandscapePatch/Private/LandscapePatchManager.cpp
@@ -196,6 +196,24 @@ bool ALandscapePatchManager::RemovePatch(TObjectPtr<ULandscapePatchComponent> Pa
 }
 
 #if WITH_EDITOR
+bool ALandscapePatchManager::IsAffectingWeightmapLayer(const FName& InLayerName) const
+{
+	for (const TSoftObjectPtr<ULandscapePatchComponent>& Component : PatchComponents)
+	{
+		if (Component.IsPending())
+		{
+			Component.LoadSynchronous();
+		}
+
+		if (Component.IsValid() && Component->IsEnabled() && Component->IsAffectingWeightmapLayer(InLayerName))
+		{
+			return true;
+		}
+	}
+
+	return false;
+}
+
 void ALandscapePatchManager::PostEditUndo()
 {
 	RequestLandscapeUpdate();
diff --git a/Engine/Plugins/Experimental/LandscapePatch/Source/LandscapePatch/Private/LandscapePatchUtil.cpp b/Engine/Plugins/Experimental/LandscapePatch/Source/LandscapePatch/Private/LandscapePatchUtil.cpp
new file mode 100644
--- /dev/null
+++ b/Engine/Plugins/Experimental/LandscapePatch/Source/LandscapePatch/Private/LandscapePatchUtil.cpp
@@ -0,0 +1,23 @@
+// Copyright Epic Games, Inc. All Rights Reserved.
+
+#include "LandscapePatchUtil.h"
+
+#include "LandscapeTexturePatchPS.h" // FSimpleTextureCopyPS
+#include "LandscapeTextureBackedRenderTarget.h"
+#include "RenderGraph.h"
+
+void UE::Landscape::PatchUtil::CopyTextureOnRenderThread(FRHICommandListImmediate& RHICmdList, 
+	const FTextureResource& Source, FTextureResource& Destination)
+{
+	using namespace UE::Landscape;
+
+	FRDGBuilder GraphBuilder(RHICmdList, RDG_EVENT_NAME("LandscapeTexturePatchCopyTexture"));
+
+	FRDGTextureRef SourceTexture = GraphBuilder.RegisterExternalTexture(CreateRenderTarget(Source.GetTexture2DRHI(), TEXT("CopySource")));
+	FRDGTextureRef DestinationTexture = GraphBuilder.RegisterExternalTexture(CreateRenderTarget(Destination.GetTexture2DRHI(), TEXT("CopyDestination")));
+
+	// All my efforts of getting CopyToResolveTarget to work without complaints have failed, so we just use our own copy shader.
+	FSimpleTextureCopyPS::AddToRenderGraph(GraphBuilder, SourceTexture, DestinationTexture);
+
+	GraphBuilder.Execute();
+}
diff --git a/Engine/Plugins/Experimental/LandscapePatch/Source/LandscapePatch/Private/LandscapeTextureBackedRenderTarget.cpp b/Engine/Plugins/Experimental/LandscapePatch/Source/LandscapePatch/Private/LandscapeTextureBackedRenderTarget.cpp
new file mode 100644
--- /dev/null
+++ b/Engine/Plugins/Experimental/LandscapePatch/Source/LandscapePatch/Private/LandscapeTextureBackedRenderTarget.cpp
@@ -0,0 +1,468 @@
+// Copyright Epic Games, Inc. All Rights Reserved.
+
+#include "LandscapeTextureBackedRenderTarget.h"
+
+#include "LandscapeDataAccess.h"
+#include "LandscapePatchUtil.h" // CopyTextureOnRenderThread
+#include "LandscapeTexturePatchPS.h"
+#include "RenderGraph.h"
+#include "TextureCompiler.h"
+#include "UObject/ObjectSaveContext.h"
+
+namespace LandscapeTextureBackedRenderTargetLocals
+{
+#if WITH_EDITOR
+	UTexture2D* CreateTexture(UObject* Parent)
+	{
+		UTexture2D* Texture = Parent ? NewObject<UTexture2D>(Parent) : NewObject<UTexture2D>();
+		Texture->SRGB = false;
+		Texture->MipGenSettings = TMGS_NoMipmaps;
+		Texture->AddressX = TA_Clamp;
+		Texture->AddressY = TA_Clamp;
+		// TODO: How do we allow lossless compression, instead of disallowing compression entirely? Just setting 
+		// LossyCompressionAmount to ETextureLossyCompressionAmount::TLCA_None is not sufficient.
+		Texture->CompressionNone = true;
+
+		return Texture;
+	}
+#endif // WITH_EDITOR
+}
+
+#if WITH_EDITOR
+void ULandscapeTextureBackedRenderTargetBase::PreSave(FObjectPreSaveContext SaveContext)
+{
+	Super::PreSave(SaveContext);
+
+	if (!bUseInternalTextureOnly)
+	{
+		CopyToInternalTexture();
+	}
+}
+
+void ULandscapeTextureBackedRenderTargetBase::PostLoad()
+{
+	Super::PostLoad();
+
+	if (!bUseInternalTextureOnly && IsValid(InternalTexture))
+	{
+		InternalTexture->ConditionalPostLoad();
+
+		RenderTarget = NewObject<UTextureRenderTarget2D>(this);
+		RenderTarget->RenderTargetFormat = GetRenderTargetFormat();
+		RenderTarget->InitAutoFormat(SizeX, SizeY);
+		RenderTarget->UpdateResourceImmediate(false);
+
+		CopyBackFromInternalTexture();
+	}
+}
+
+void ULandscapeTextureBackedRenderTargetBase::PreDuplicate(FObjectDuplicationParameters& DupParams)
+{
+	Super::PreDuplicate(DupParams);
+
+	if (!bUseInternalTextureOnly)
+	{
+		CopyToInternalTexture();
+	}
+}
+
+// Called when serializing to text for copy/paste
+void ULandscapeTextureBackedRenderTargetBase::ExportCustomProperties(FOutputDevice& Out, uint32 Indent)
+{
+	Super::ExportCustomProperties(Out, Indent);
+
+	if (!bUseInternalTextureOnly)
+	{
+		CopyToInternalTexture();
+	}
+}
+
+void ULandscapeTextureBackedRenderTargetBase::PostEditImport()
+{
+	Super::PostEditImport();
+
+	if (!bUseInternalTextureOnly && IsValid(InternalTexture))
+	{
+		RenderTarget = NewObject<UTextureRenderTarget2D>(this);
+		RenderTarget->RenderTargetFormat = GetRenderTargetFormat();
+		RenderTarget->InitAutoFormat(SizeX, SizeY);
+		RenderTarget->UpdateResourceImmediate(false);
+
+		CopyBackFromInternalTexture();
+	}
+}
+#endif // WITH_EDITOR
+
+void ULandscapeTextureBackedRenderTargetBase::SetUseInternalTextureOnly(bool bUseInternalTextureOnlyIn, bool bCopyExisting)
+{
+#if WITH_EDITOR
+	if (bUseInternalTextureOnly == bUseInternalTextureOnlyIn)
+	{
+		return;
+	}
+
+	Modify();
+	bUseInternalTextureOnly = bUseInternalTextureOnlyIn;
+
+	if (bUseInternalTextureOnly)
+	{
+		// We're no longer using the render target, so copy it to internal texture
+		if (IsValid(RenderTarget) && bCopyExisting)
+		{
+			Modify();
+			CopyToInternalTexture();
+		}
+		RenderTarget = nullptr;
+	}
+	else
+	{
+		// We're back to using the render target, so initialize it from internal texture.
+
+		if (IsValid(InternalTexture)) // if initialized
+		{
+			RenderTarget = NewObject<UTextureRenderTarget2D>(this);
+			RenderTarget->RenderTargetFormat = GetRenderTargetFormat();
+			RenderTarget->InitAutoFormat(SizeX, SizeY);
+			RenderTarget->UpdateResourceImmediate(false);
+
+			if (bCopyExisting)
+			{
+				CopyBackFromInternalTexture();
+			}
+		}
+	}
+#endif // WITH_EDITOR
+}
+
+void ULandscapeTextureBackedRenderTargetBase::SetResolution(int32 SizeXIn, int32 SizeYIn)
+{
+#if WITH_EDITOR
+	if (SizeXIn == SizeX && SizeYIn == SizeY)
+	{
+		return;
+	}
+
+	Modify();
+	SizeX = SizeXIn;
+	SizeY = SizeYIn;
+	if (IsValid(RenderTarget))
+	{
+		RenderTarget->Modify();
+		RenderTarget->InitAutoFormat(SizeX, SizeY);
+		RenderTarget->UpdateResourceImmediate();
+	}
+	if (IsValid(InternalTexture))
+	{
+		InternalTexture->Modify();
+		InternalTexture->Source.Init(SizeX, SizeY, 1, 1, GetInternalTextureFormat());
+		InternalTexture->UpdateResource();
+	}
+#endif // WITH_EDITOR
+}
+
+void ULandscapeTextureBackedRenderTargetBase::Initialize()
+{
+#if WITH_EDITOR
+	using namespace LandscapeTextureBackedRenderTargetLocals;
+
+	if (!IsValid(InternalTexture) || !InternalTexture->GetResource())
+	{
+		InternalTexture = CreateTexture(this);
+	}
+
+	InternalTexture->Source.Init(SizeX, SizeY, 1, 1, GetInternalTextureFormat());
+	InternalTexture->UpdateResource();
+
+	if (bUseInternalTextureOnly)
+	{
+		RenderTarget = nullptr;
+	}
+	else
+	{
+		if (!IsValid(RenderTarget))
+		{
+			RenderTarget = NewObject<UTextureRenderTarget2D>(this);
+		}
+
+		RenderTarget->RenderTargetFormat = GetRenderTargetFormat();
+		RenderTarget->InitAutoFormat(SizeX, SizeY);
+		RenderTarget->UpdateResourceImmediate();
+	}
+#endif // WITH_EDITOR
+}
+
+
+
+void ULandscapeWeightTextureBackedRenderTarget::SetUseAlphaChannel(bool bUseAlphaChannelIn)
+{
+#if WITH_EDITOR
+	if (bUseAlphaChannel == bUseAlphaChannelIn)
+	{
+		return;
+	}
+	Modify();
+	bUseAlphaChannel = bUseAlphaChannelIn;
+
+	// TODO: We could try to copy the non-alpha channel data across, but that is messy with
+	// an unclear amount of benefit. It would seem odd for a user to want to write one
+	// way but then discard/add the alpha channel afterward. Still, could revisit later.
+
+	if (IsValid(RenderTarget))
+	{
+		RenderTarget->Modify();
+		RenderTarget->RenderTargetFormat = GetRenderTargetFormat();
+		RenderTarget->InitAutoFormat(SizeX, SizeY);
+		RenderTarget->UpdateResourceImmediate();
+	}
+	if (IsValid(InternalTexture))
+	{
+		InternalTexture->Modify();
+		InternalTexture->Source.Init(SizeX, SizeY, 1, 1, GetInternalTextureFormat());
+		InternalTexture->UpdateResource();
+	}
+#endif // WITH_EDITOR
+}
+
+ETextureRenderTargetFormat ULandscapeWeightTextureBackedRenderTarget::GetRenderTargetFormat()
+{
+	return bUseAlphaChannel ? ETextureRenderTargetFormat::RTF_RGBA8 : ETextureRenderTargetFormat::RTF_R8;
+}
+
+ETextureSourceFormat ULandscapeWeightTextureBackedRenderTarget::GetInternalTextureFormat()
+{
+	return bUseAlphaChannel ? ETextureSourceFormat::TSF_BGRA8 : ETextureSourceFormat::TSF_G8;
+}
+
+void ULandscapeWeightTextureBackedRenderTarget::CopyToInternalTexture()
+{
+#if WITH_EDITOR
+	using namespace LandscapeTextureBackedRenderTargetLocals;
+
+	if (!ensure(IsValid(RenderTarget)))
+	{
+		return;
+	}
+
+	if (!ensure(IsValid(InternalTexture) && InternalTexture->GetResource()))
+	{
+		Modify();
+		InternalTexture = CreateTexture(this);
+		// The sizing and format doesn't matter because the UpdateTexture2D call below will deal with it.
+	}
+	else
+	{
+		InternalTexture->Modify();
+	}
+
+	RenderTarget->UpdateTexture2D(InternalTexture, GetInternalTextureFormat());
+	InternalTexture->UpdateResource();
+#endif // WITH_EDITOR
+}
+
+void ULandscapeWeightTextureBackedRenderTarget::CopyBackFromInternalTexture()
+{
+#if WITH_EDITOR
+	using namespace LandscapeTextureBackedRenderTargetLocals;
+
+	if (!ensure(IsValid(InternalTexture)))
+	{
+		return;
+	}
+	InternalTexture->UpdateResource();
+	FTextureCompilingManager::Get().FinishCompilation({ InternalTexture });
+
+	if (!ensure(InternalTexture->GetResource()))
+	{
+		return;
+	}
+
+	bool bCreatedNewRenderTarget = false;
+	if (!ensure(IsValid(RenderTarget)))
+	{
+		Modify();
+		RenderTarget = NewObject<UTextureRenderTarget2D>(this);
+		bCreatedNewRenderTarget = true;
+	}
+
+	if (!ensure(InternalTexture->GetResource()->GetSizeX() == RenderTarget->SizeX
+		&& InternalTexture->GetResource()->GetSizeY() == RenderTarget->SizeY
+		&& RenderTarget->RenderTargetFormat == GetRenderTargetFormat()))
+	{
+		if (!bCreatedNewRenderTarget)
+		{
+			RenderTarget->Modify();
+		}
+		RenderTarget->RenderTargetFormat = GetRenderTargetFormat();
+		RenderTarget->InitAutoFormat(InternalTexture->GetResource()->GetSizeX(), InternalTexture->GetResource()->GetSizeY());
+		RenderTarget->UpdateResourceImmediate(false);
+	}
+
+	FTextureResource* Source = InternalTexture->GetResource();
+	FTextureResource* Destination = RenderTarget->GetResource();
+
+	ENQUEUE_RENDER_COMMAND(LandscapeTextureHeightPatchRTToTexture)(
+		[Source, Destination](FRHICommandListImmediate& RHICmdList)
+		{
+			UE::Landscape::PatchUtil::CopyTextureOnRenderThread(RHICmdList, *Source, *Destination);
+		});
+#endif // WITH_EDITOR
+}
+
+
+void ULandscapeHeightTextureBackedRenderTarget::SetFormat(ETextureRenderTargetFormat FormatToUse)
+{
+#if WITH_EDITOR
+	if (RenderTargetFormat == FormatToUse)
+	{
+		return;
+	}
+	Modify();
+	RenderTargetFormat = FormatToUse;
+
+	// We could try to copy over existing data, but that is not worth it.
+
+	if (IsValid(RenderTarget))
+	{
+		RenderTarget->Modify();
+		RenderTarget->RenderTargetFormat = GetRenderTargetFormat();
+		RenderTarget->InitAutoFormat(SizeX, SizeY);
+		RenderTarget->UpdateResourceImmediate();
+	}
+#endif // WITH_EDITOR
+}
+
+void ULandscapeHeightTextureBackedRenderTarget::CopyToInternalTexture()
+{
+#if WITH_EDITOR
+	using namespace LandscapeTextureBackedRenderTargetLocals;
+
+	if (!ensure(IsValid(RenderTarget)))
+	{
+		return;
+	}
+
+	if (!ensure(IsValid(InternalTexture) && InternalTexture->GetResource()))
+	{
+		Modify();
+		InternalTexture = CreateTexture(this);
+	}
+	else
+	{
+		InternalTexture->Modify();
+	}
+
+	UTextureRenderTarget2D* NativeEncodingRenderTarget = RenderTarget;
+
+	// If the format doesn't match the format that we use generally for our internal texture, save the patch in our native
+	// height format, applying whatever scale/offset is relevant. The stored texture thus ends up being the native equivalent
+	// (with scale 1 and offset 0). This is easier than trying to support various kinds of RT-to-texture conversions.
+	if (NativeEncodingRenderTarget->RenderTargetFormat != ETextureRenderTargetFormat::RTF_RGBA8)
+	{
+		// We need a temporary render target to write the converted result, then we'll copy that to the texture.
+		NativeEncodingRenderTarget = NewObject<UTextureRenderTarget2D>(this);
+		NativeEncodingRenderTarget->RenderTargetFormat = ETextureRenderTargetFormat::RTF_RGBA8;
+		NativeEncodingRenderTarget->InitAutoFormat(RenderTarget->SizeX, RenderTarget->SizeY);
+		NativeEncodingRenderTarget->UpdateResourceImmediate(false);
+
+		FTextureResource* Source = RenderTarget->GetResource();
+		FTextureResource* Destination = NativeEncodingRenderTarget->GetResource();
+
+		ENQUEUE_RENDER_COMMAND(LandscapeTextureHeightPatchRTToTexture)(
+			[this, Source, Destination](FRHICommandListImmediate& RHICmdList)
+			{
+				using namespace UE::Landscape;
+
+				FRDGBuilder GraphBuilder(RHICmdList, RDG_EVENT_NAME("LandscapeTextureHeightPatchConvertToNative"));
+
+				FRDGTextureRef SourceTexture = GraphBuilder.RegisterExternalTexture(CreateRenderTarget(Source->GetTexture2DRHI(), TEXT("ConversionSource")));
+				FRDGTextureRef DestinationTexture = GraphBuilder.RegisterExternalTexture(CreateRenderTarget(Destination->GetTexture2DRHI(), TEXT("ConversionDestination")));
+
+				FConvertToNativeLandscapePatchPS::AddToRenderGraph(GraphBuilder, SourceTexture, DestinationTexture, ConversionParams);
+
+				GraphBuilder.Execute();
+			});
+	}
+
+	// TODO: The header for this requires the texture to be square power of 2, but it actually doesn't seem to
+	// be an enforced requirement. If that changes, we'll need our own ReadPixels followed by locking a mip
+	// and writing to it.
+	// This call does a flush for us, so the render target should be updated.
+	NativeEncodingRenderTarget->UpdateTexture2D(InternalTexture, GetInternalTextureFormat());
+	InternalTexture->UpdateResource();
+#endif // WITH_EDITOR
+}
+
+void ULandscapeHeightTextureBackedRenderTarget::CopyBackFromInternalTexture()
+{
+#if WITH_EDITOR
+	using namespace LandscapeTextureBackedRenderTargetLocals;
+
+	if (!ensure(IsValid(InternalTexture)))
+	{
+		return;
+	}
+
+	InternalTexture->UpdateResource();
+	FTextureCompilingManager::Get().FinishCompilation({ InternalTexture });
+
+	if (!ensure(InternalTexture->GetResource()))
+	{
+		return;
+	}
+
+	bool bCreatedNewRenderTarget = false;
+	if (!ensure(IsValid(RenderTarget)))
+	{
+		Modify();
+		RenderTarget = NewObject<UTextureRenderTarget2D>(this);
+		bCreatedNewRenderTarget = true;
+	}
+
+	if (!ensure(InternalTexture->GetResource()->GetSizeX() == RenderTarget->SizeX
+		&& InternalTexture->GetResource()->GetSizeY() == RenderTarget->SizeY
+		&& RenderTarget->RenderTargetFormat == GetRenderTargetFormat()))
+	{
+		if (!bCreatedNewRenderTarget)
+		{
+			RenderTarget->Modify();
+		}
+		RenderTarget->RenderTargetFormat = GetRenderTargetFormat();
+		RenderTarget->InitAutoFormat(InternalTexture->GetResource()->GetSizeX(), InternalTexture->GetResource()->GetSizeY());
+	}
+
+	FTextureResource* Source = InternalTexture->GetResource();
+	FTextureResource* Destination = RenderTarget->GetResource();
+	if (!ensure(Source && Destination))
+	{
+		return;
+	}
+
+	// If we're in a different format, we need to "un-bake" the height from the texture.
+	if (RenderTarget->RenderTargetFormat != ETextureRenderTargetFormat::RTF_RGBA8)
+	{
+		ENQUEUE_RENDER_COMMAND(LandscapeTextureHeightPatchRTToTexture)(
+			[this, Source, Destination](FRHICommandListImmediate& RHICmdList)
+			{
+				using namespace UE::Landscape;
+
+				FRDGBuilder GraphBuilder(RHICmdList, RDG_EVENT_NAME("LandscapeTextureHeightPatchConvertFromNative"));
+
+				FRDGTextureRef SourceTexture = GraphBuilder.RegisterExternalTexture(CreateRenderTarget(Source->GetTexture2DRHI(), TEXT("ConversionSource")));
+				FRDGTextureRef DestinationTexture = GraphBuilder.RegisterExternalTexture(CreateRenderTarget(Destination->GetTexture2DRHI(), TEXT("ConversionDestination")));
+
+				FConvertBackFromNativeLandscapePatchPS::AddToRenderGraph(GraphBuilder, SourceTexture, DestinationTexture, ConversionParams);
+
+				GraphBuilder.Execute();
+			});
+	}
+	else
+	{
+		// When formats match, we can just copy back and forth.
+		ENQUEUE_RENDER_COMMAND(LandscapeTextureHeightPatchRTToTexture)(
+			[Source, Destination](FRHICommandListImmediate& RHICmdList)
+			{
+				UE::Landscape::PatchUtil::CopyTextureOnRenderThread(RHICmdList, *Source, *Destination);
+			});
+	}
+#endif // WITH_EDITOR
+}
\ No newline at end of file
diff --git a/Engine/Plugins/Experimental/LandscapePatch/Source/LandscapePatch/Private/LandscapeTextureHeightPatch.cpp b/Engine/Plugins/Experimental/LandscapePatch/Source/LandscapePatch/Private/LandscapeTextureHeightPatch.cpp
--- a/Engine/Plugins/Experimental/LandscapePatch/Source/LandscapePatch/Private/LandscapeTextureHeightPatch.cpp
+++ b/Engine/Plugins/Experimental/LandscapePatch/Source/LandscapePatch/Private/LandscapeTextureHeightPatch.cpp
@@ -72,7 +72,7 @@ void ULandscapeTextureHeightPatch::ConvertInternalRenderTargetBackFromNativeText
 	// If we're in a different format, we need to "un-bake" the height from the texture.
 	if (InternalRenderTarget->RenderTargetFormat != ETextureRenderTargetFormat::RTF_RGBA8)
 	{
-		FConvertToNativeLandscapePatchParams ConversionParams = GetConversionParams();
+		FLandscapeHeightPatchConvertToNativeParams ConversionParams = GetConversionParams();
 		if (bLoading)
 		{
 			ConversionParams.HeightScale = SavedConversionHeightScale;
@@ -142,7 +142,7 @@ void ULandscapeTextureHeightPatch::ConvertInternalRenderTargetToNativeTexture(bo
 		FTextureResource* Source = InternalRenderTarget->GetResource();
 		FTextureResource* Destination = NativeEncodingRenderTarget->GetResource();
 
-		FConvertToNativeLandscapePatchParams ConversionParams = GetConversionParams();
+		FLandscapeHeightPatchConvertToNativeParams ConversionParams = GetConversionParams();
 		SavedConversionHeightScale = ConversionParams.HeightScale;
 
 		ENQUEUE_RENDER_COMMAND(LandscapeTextureHeightPatchRTToTexture)(
@@ -176,11 +176,11 @@ void ULandscapeTextureHeightPatch::ConvertInternalRenderTargetToNativeTexture(bo
 	}
 }
 
-UE::Landscape::FConvertToNativeLandscapePatchParams ULandscapeTextureHeightPatch::GetConversionParams()
+FLandscapeHeightPatchConvertToNativeParams ULandscapeTextureHeightPatch::GetConversionParams()
 {
 	// When doing conversions, we bake into a height in the same way that we do when applying the patch.
 
-	UE::Landscape::FConvertToNativeLandscapePatchParams ConversionParams;
+	FLandscapeHeightPatchConvertToNativeParams ConversionParams;
 	ConversionParams.ZeroInEncoding = EncodingSettings.ZeroInEncoding;
 
 	double LandscapeHeightScale = Landscape.IsValid() ? Landscape->GetTransform().GetScale3D().Z : 1;
@@ -462,8 +462,8 @@ UTextureRenderTarget2D* ULandscapeTextureHeightPatch::Render_Native(bool bIsHeig
 {
 	using namespace UE::Landscape;
 
-	if (!ensure(PatchManager.IsValid()) || 
-		(SourceMode == ELandscapeTexturePatchSourceMode::TextureBackedRenderTarget && (!IsValid(InternalRenderTarget)))
+	if (!ensure(PatchManager.IsValid()) || SourceMode == ELandscapeTexturePatchSourceMode::None
+		|| (SourceMode == ELandscapeTexturePatchSourceMode::TextureBackedRenderTarget && (!IsValid(InternalRenderTarget)))
 		|| (SourceMode == ELandscapeTexturePatchSourceMode::TextureAsset && (!IsValid(TextureAsset) 
 			|| !ensureMsgf(TextureAsset->VirtualTextureStreaming == 0, TEXT("ULandscapeTextureHeightPatch: Virtual textures are not supported"))))
 		|| (SourceMode == ELandscapeTexturePatchSourceMode::InternalTexture && (!IsValid(InternalTexture))))
diff --git a/Engine/Plugins/Experimental/LandscapePatch/Source/LandscapePatch/Private/LandscapeTextureHeightPatchPS.cpp b/Engine/Plugins/Experimental/LandscapePatch/Source/LandscapePatch/Private/LandscapeTextureHeightPatchPS.cpp
deleted file mode 100644
diff --git a/Engine/Plugins/Experimental/LandscapePatch/Source/LandscapePatch/Private/LandscapeTexturePatch.cpp b/Engine/Plugins/Experimental/LandscapePatch/Source/LandscapePatch/Private/LandscapeTexturePatch.cpp
new file mode 100644
--- /dev/null
+++ b/Engine/Plugins/Experimental/LandscapePatch/Source/LandscapePatch/Private/LandscapeTexturePatch.cpp
@@ -0,0 +1,1332 @@
+// Copyright Epic Games, Inc. All Rights Reserved.
+
+#include "LandscapeTexturePatch.h"
+
+#include "Landscape.h"
+#include "LandscapeDataAccess.h"
+#include "LandscapeLayerInfoObject.h" // VisibilityLayer->LayerName
+#include "LandscapeInfo.h"
+#include "LandscapePatchLogging.h"
+#include "LandscapePatchManager.h"
+#include "LandscapePatchUtil.h" // CopyTextureOnRenderThread
+#include "LandscapeTexturePatchPS.h"
+#include "MathUtil.h"
+#include "RenderGraph.h"
+
+namespace LandscapeTexturePatchLocals
+{
+#if WITH_EDITOR
+	template <typename TextureBackedRTType>
+	void TransitionSourceMode(ELandscapeTexturePatchSourceMode OldMode, ELandscapeTexturePatchSourceMode NewMode,
+		TObjectPtr<UTexture>& TextureAsset, TObjectPtr<TextureBackedRTType>& InternalData,
+		TUniqueFunction<TextureBackedRTType* ()> InternalDataBuilder)
+	{
+		if (NewMode == ELandscapeTexturePatchSourceMode::None)
+		{
+			TextureAsset = nullptr;
+			InternalData = nullptr;
+		}
+		else if (NewMode == ELandscapeTexturePatchSourceMode::TextureAsset)
+		{
+			InternalData = nullptr;
+		}
+		else // new mode is internal texture or render target
+		{
+			bool bWillUseTextureOnly = (NewMode == ELandscapeTexturePatchSourceMode::InternalTexture);
+			bool bNeedToCopyTextureAsset = (OldMode == ELandscapeTexturePatchSourceMode::TextureAsset
+				&& IsValid(TextureAsset) && TextureAsset->GetResource());
+
+			if (!InternalData)
+			{
+				InternalData = InternalDataBuilder();
+				InternalData->SetUseInternalTextureOnly(bWillUseTextureOnly && !bNeedToCopyTextureAsset);
+				InternalData->Initialize();
+			}
+			else
+			{
+				InternalData->Modify();
+			}
+
+			InternalData->SetUseInternalTextureOnly(bWillUseTextureOnly && !bNeedToCopyTextureAsset);
+			if (bNeedToCopyTextureAsset)
+			{
+				// Copy the currently set texture asset to our render target
+				FTextureResource* Source = TextureAsset->GetResource();
+				FTextureResource* Destination = InternalData->GetRenderTarget()->GetResource();
+
+				ENQUEUE_RENDER_COMMAND(LandscapeTextureHeightPatchRTToTexture)(
+					[Source, Destination](FRHICommandListImmediate& RHICmdList)
+					{
+						UE::Landscape::PatchUtil::CopyTextureOnRenderThread(RHICmdList, *Source, *Destination);
+					});
+			}
+
+			// Note that the duplicate SetUseInternalTextureOnly calls (in cases where we don't need to copy the texture asset)
+			// are fine because they don't do anything.
+			InternalData->SetUseInternalTextureOnly(bWillUseTextureOnly);
+
+			TextureAsset = nullptr;
+		}
+	}
+#endif // WITH_EDITOR
+}
+
+#if WITH_EDITOR
+UTextureRenderTarget2D* ULandscapeTexturePatch::Render_Native(bool bIsHeightmap,
+	UTextureRenderTarget2D* InCombinedResult,
+	const FName& InWeightmapLayerName)
+{
+	using namespace UE::Landscape;
+
+	if (bIsHeightmap)
+	{
+		if (bReinitializeHeightOnNextRender)
+		{
+			bReinitializeHeightOnNextRender = false;
+			ReinitializeHeight(InCombinedResult);
+			return InCombinedResult;
+		}
+		else
+		{
+			return ApplyToHeightmap(InCombinedResult);
+		}
+	}
+	else
+	{
+		// Try to find the weight patch
+		ULandscapeWeightPatchTextureInfo* WeightPatchInfo = nullptr;
+
+		for (TObjectPtr<ULandscapeWeightPatchTextureInfo> WeightPatchEntry : WeightPatches)
+		{
+			if (WeightPatchEntry->WeightmapLayerName == InWeightmapLayerName)
+			{
+				WeightPatchInfo = WeightPatchEntry;
+				break;
+			}
+		}
+
+		if (!WeightPatchInfo)
+		{
+			return InCombinedResult;
+		}
+
+		if (WeightPatchInfo->bReinitializeOnNextRender)
+		{
+			WeightPatchInfo->bReinitializeOnNextRender = false;
+			ReinitializeWeightPatch(WeightPatchInfo, InCombinedResult);
+			return InCombinedResult;
+		}
+		else
+		{
+			return ApplyToWeightmap(WeightPatchInfo, InCombinedResult);
+		}
+	}
+}
+
+UTextureRenderTarget2D* ULandscapeTexturePatch::ApplyToHeightmap(UTextureRenderTarget2D* InCombinedResult)
+{
+	using namespace UE::Landscape;
+
+	// Get the source of our height patch
+	UTexture* PatchUObject = nullptr;
+	switch (HeightSourceMode)
+	{
+	case ELandscapeTexturePatchSourceMode::None:
+		return InCombinedResult;
+	case ELandscapeTexturePatchSourceMode::InternalTexture:
+		PatchUObject = GetHeightInternalTexture();
+		break;
+	case ELandscapeTexturePatchSourceMode::TextureBackedRenderTarget:
+		PatchUObject = GetHeightRenderTarget();
+		break;
+	case ELandscapeTexturePatchSourceMode::TextureAsset:
+
+		if (IsValid(HeightTextureAsset) && !ensureMsgf(HeightTextureAsset->VirtualTextureStreaming == 0,
+			TEXT("ULandscapeTexturePatch: Virtual textures are not supported")))
+		{
+			return InCombinedResult;
+		}
+		PatchUObject = HeightTextureAsset;
+		break;
+	default:
+		ensure(false);
+	}
+
+	if (!IsValid(PatchUObject))
+	{
+		return InCombinedResult;
+	}
+
+	FTextureResource* Patch = PatchUObject->GetResource();
+	if (!Patch)
+	{
+		return InCombinedResult;
+	}
+
+	// Go ahead and pack everything into a copy of the param struct so we don't have to capture everything
+	// individually in the lambda below.
+	FApplyLandscapeTextureHeightPatchPS::FParameters ShaderParamsToCopy;
+	FIntRect DestinationBounds;
+	GetHeightShaderParams(FIntPoint(Patch->GetSizeX(), Patch->GetSizeY()), FIntPoint(InCombinedResult->SizeX, InCombinedResult->SizeY), ShaderParamsToCopy, DestinationBounds);
+
+	if (DestinationBounds.IsEmpty())
+	{
+		// Patch must be outside the landscape.
+		return InCombinedResult;
+	}
+
+	ENQUEUE_RENDER_COMMAND(LandscapeTextureHeightPatch)([InCombinedResult, ShaderParamsToCopy, Patch, DestinationBounds](FRHICommandListImmediate& RHICmdList)
+		{
+			TRACE_CPUPROFILER_EVENT_SCOPE(LandscapeTextureHeightPatch_Render);
+
+			FRDGBuilder GraphBuilder(RHICmdList, RDG_EVENT_NAME("ApplyTextureHeightPatch"));
+
+			TRefCountPtr<IPooledRenderTarget> DestinationRenderTarget = CreateRenderTarget(InCombinedResult->GetResource()->GetTexture2DRHI(), TEXT("LandscapeTextureHeightPatchOutput"));
+			FRDGTextureRef DestinationTexture = GraphBuilder.RegisterExternalTexture(DestinationRenderTarget);
+
+			// Make a copy of our heightmap input so we can read and write at the same time (needed for blending)
+			FRDGTextureRef InputCopy = GraphBuilder.CreateTexture(DestinationTexture->Desc, TEXT("LandscapeTextureHeightPatchInputCopy"));
+
+			FRHICopyTextureInfo CopyTextureInfo;
+			CopyTextureInfo.NumMips = 1;
+			CopyTextureInfo.Size = FIntVector(DestinationTexture->Desc.GetSize().X, DestinationTexture->Desc.GetSize().Y, 0);
+			AddCopyTexturePass(GraphBuilder, DestinationTexture, InputCopy, CopyTextureInfo);
+
+			FApplyLandscapeTextureHeightPatchPS::FParameters* ShaderParams =
+				GraphBuilder.AllocParameters<FApplyLandscapeTextureHeightPatchPS::FParameters>();
+			*ShaderParams = ShaderParamsToCopy;
+
+			TRefCountPtr<IPooledRenderTarget> PatchRenderTarget = CreateRenderTarget(Patch->GetTexture2DRHI(), TEXT("LandscapeTextureHeightPatch"));
+			FRDGTextureRef PatchTexture = GraphBuilder.RegisterExternalTexture(PatchRenderTarget);
+			FRDGTextureSRVRef PatchSRV = GraphBuilder.CreateSRV(FRDGTextureSRVDesc::CreateForMipLevel(PatchTexture, 0));
+			ShaderParams->InHeightPatch = PatchSRV;
+			ShaderParams->InHeightPatchSampler = TStaticSamplerState<SF_Bilinear, AM_Clamp, AM_Clamp>::GetRHI();
+
+			FRDGTextureSRVRef InputCopySRV = GraphBuilder.CreateSRV(FRDGTextureSRVDesc::CreateForMipLevel(InputCopy, 0));
+			ShaderParams->InSourceHeightmap = InputCopySRV;
+
+			ShaderParams->RenderTargets[0] = FRenderTargetBinding(DestinationTexture, ERenderTargetLoadAction::ENoAction, /*InMipIndex = */0);
+
+			FApplyLandscapeTextureHeightPatchPS::AddToRenderGraph(GraphBuilder, ShaderParams, DestinationBounds);
+
+			GraphBuilder.Execute();
+		});
+
+	return InCombinedResult;
+}
+
+UTextureRenderTarget2D* ULandscapeTexturePatch::ApplyToWeightmap(ULandscapeWeightPatchTextureInfo* PatchInfo, UTextureRenderTarget2D* InCombinedResult)
+{
+	using namespace UE::Landscape;
+
+	if (!PatchInfo)
+	{
+		return InCombinedResult;
+	}
+
+	UTexture* PatchUObject = nullptr;
+
+	switch (PatchInfo->SourceMode)
+	{
+	case ELandscapeTexturePatchSourceMode::None:
+		return InCombinedResult;
+	case ELandscapeTexturePatchSourceMode::InternalTexture:
+		PatchUObject = PatchInfo->InternalData->GetInternalTexture();
+		break;
+	case ELandscapeTexturePatchSourceMode::TextureBackedRenderTarget:
+		PatchUObject = PatchInfo->InternalData->GetRenderTarget();
+		break;
+	case ELandscapeTexturePatchSourceMode::TextureAsset:
+		if (IsValid(PatchInfo->TextureAsset) && !ensureMsgf(PatchInfo->TextureAsset->VirtualTextureStreaming == 0,
+			TEXT("ULandscapeTexturePatch: Virtual textures are not supported")))
+		{
+			return InCombinedResult;
+		}
+		PatchUObject = PatchInfo->TextureAsset;
+		break;
+	default:
+		ensure(false);
+	}
+
+	if (!IsValid(PatchUObject))
+	{
+		return InCombinedResult;
+	}
+
+	FTextureResource* Patch = PatchUObject->GetResource();
+	if (!Patch)
+	{
+		return InCombinedResult;
+	}
+
+	// Go ahead and pack everything into a copy of the param struct so we don't have to capture everything
+	// individually in the lambda below.
+	FApplyLandscapeTextureWeightPatchPS::FParameters ShaderParamsToCopy;
+	FIntRect DestinationBounds;
+
+	GetWeightShaderParams(FIntPoint(Patch->GetSizeX(), Patch->GetSizeY()), FIntPoint(InCombinedResult->SizeX, InCombinedResult->SizeY), 
+		PatchInfo, ShaderParamsToCopy, DestinationBounds);
+
+	if (DestinationBounds.IsEmpty())
+	{
+		// Patch must be outside the landscape.
+		return InCombinedResult;
+	}
+
+	ENQUEUE_RENDER_COMMAND(LandscapeTextureHeightPatch)([InCombinedResult, ShaderParamsToCopy, Patch, DestinationBounds](FRHICommandListImmediate& RHICmdList)
+		{
+			TRACE_CPUPROFILER_EVENT_SCOPE(LandscapeTextureHeightPatch_Render);
+
+			FRDGBuilder GraphBuilder(RHICmdList, RDG_EVENT_NAME("ApplyTextureHeightPatch"));
+
+			TRefCountPtr<IPooledRenderTarget> DestinationRenderTarget = CreateRenderTarget(InCombinedResult->GetResource()->GetTexture2DRHI(), TEXT("LandscapeTextureWeightPatchOutput"));
+			FRDGTextureRef DestinationTexture = GraphBuilder.RegisterExternalTexture(DestinationRenderTarget);
+
+			// Make a copy of our heightmap input so we can read and write at the same time (needed for blending)
+			FRDGTextureRef InputCopy = GraphBuilder.CreateTexture(DestinationTexture->Desc, TEXT("LandscapeTextureWeightPatchInputCopy"));
+
+			FRHICopyTextureInfo CopyTextureInfo;
+			CopyTextureInfo.NumMips = 1;
+			CopyTextureInfo.Size = FIntVector(DestinationTexture->Desc.GetSize().X, DestinationTexture->Desc.GetSize().Y, 0);
+			AddCopyTexturePass(GraphBuilder, DestinationTexture, InputCopy, CopyTextureInfo);
+
+			FApplyLandscapeTextureWeightPatchPS::FParameters* ShaderParams =
+				GraphBuilder.AllocParameters<FApplyLandscapeTextureWeightPatchPS::FParameters>();
+			*ShaderParams = ShaderParamsToCopy;
+
+			TRefCountPtr<IPooledRenderTarget> PatchRenderTarget = CreateRenderTarget(Patch->GetTexture2DRHI(), TEXT("LandscapeTextureWeightPatch"));
+			FRDGTextureRef PatchTexture = GraphBuilder.RegisterExternalTexture(PatchRenderTarget);
+			FRDGTextureSRVRef PatchSRV = GraphBuilder.CreateSRV(FRDGTextureSRVDesc::CreateForMipLevel(PatchTexture, 0));
+			ShaderParams->InWeightPatch = PatchSRV;
+			ShaderParams->InWeightPatchSampler = TStaticSamplerState<SF_Bilinear, AM_Clamp, AM_Clamp>::GetRHI();
+
+			FRDGTextureSRVRef InputCopySRV = GraphBuilder.CreateSRV(FRDGTextureSRVDesc::CreateForMipLevel(InputCopy, 0));
+			ShaderParams->InSourceWeightmap = InputCopySRV;
+
+			ShaderParams->RenderTargets[0] = FRenderTargetBinding(DestinationTexture, ERenderTargetLoadAction::ENoAction, /*InMipIndex = */0);
+
+			FApplyLandscapeTextureWeightPatchPS::AddToRenderGraph(GraphBuilder, ShaderParams, DestinationBounds);
+
+			GraphBuilder.Execute();
+		});
+
+	return InCombinedResult;
+}
+
+void ULandscapeTexturePatch::GetCommonShaderParams(const FIntPoint& SourceResolutionIn, const FIntPoint& DestinationResolutionIn, 
+	FTransform& PatchToWorldOut, FVector2f& PatchWorldDimensionsOut, FMatrix44f& HeightmapToPatchOut, FIntRect& DestinationBoundsOut, 
+	FVector2f& EdgeUVDeadBorderOut, float& FalloffWorldMarginOut) const
+{
+	PatchToWorldOut = GetPatchToWorldTransform();
+
+	FVector2D FullPatchDimensions = GetFullUnscaledWorldSize();
+	PatchWorldDimensionsOut = FVector2f(FullPatchDimensions);
+
+	FTransform FromPatchUVToPatch(FQuat4d::Identity, FVector3d(-FullPatchDimensions.X / 2, -FullPatchDimensions.Y / 2, 0),
+		FVector3d(FullPatchDimensions.X, FullPatchDimensions.Y, 1));
+	FMatrix44d PatchLocalToUVs = FromPatchUVToPatch.ToInverseMatrixWithScale();
+
+	FTransform LandscapeHeightmapToWorld = PatchManager->GetHeightmapCoordsToWorld();
+	FMatrix44d LandscapeToWorld = LandscapeHeightmapToWorld.ToMatrixWithScale();
+
+	FMatrix44d WorldToPatch = PatchToWorldOut.ToInverseMatrixWithScale();
+
+	// In unreal, matrix composition is done by multiplying the subsequent ones on the right, and the result
+	// is transpose of what our shader will expect (because unreal right multiplies vectors by matrices).
+	FMatrix44d LandscapeToPatchUVTransposed = LandscapeToWorld * WorldToPatch * PatchLocalToUVs;
+	HeightmapToPatchOut = (FMatrix44f)LandscapeToPatchUVTransposed.GetTransposed();
+
+
+	// Get the output bounds, which are used to limit the amount of landscape pixels we have to process. 
+	// To get them, convert all of the corners into heightmap 2d coordinates and get the bounding box.
+	auto PatchUVToHeightmap2DCoordinates = [&PatchToWorldOut, &FromPatchUVToPatch, &LandscapeHeightmapToWorld](const FVector2f& UV)
+	{
+		FVector WorldPosition = PatchToWorldOut.TransformPosition(
+			FromPatchUVToPatch.TransformPosition(FVector(UV.X, UV.Y, 0)));
+		FVector HeightmapCoordinates = LandscapeHeightmapToWorld.InverseTransformPosition(WorldPosition);
+		return FVector2d(HeightmapCoordinates.X, HeightmapCoordinates.Y);
+	};
+	FBox2D FloatBounds(ForceInit);
+	FloatBounds += PatchUVToHeightmap2DCoordinates(FVector2f(0, 0));
+	FloatBounds += PatchUVToHeightmap2DCoordinates(FVector2f(0, 1));
+	FloatBounds += PatchUVToHeightmap2DCoordinates(FVector2f(1, 0));
+	FloatBounds += PatchUVToHeightmap2DCoordinates(FVector2f(1, 1));
+
+	DestinationBoundsOut = FIntRect(
+		FMath::Clamp(FMath::Floor(FloatBounds.Min.X), 0, DestinationResolutionIn.X - 1),
+		FMath::Clamp(FMath::Floor(FloatBounds.Min.Y), 0, DestinationResolutionIn.Y - 1),
+		FMath::Clamp(FMath::CeilToInt(FloatBounds.Max.X) + 1, 0, DestinationResolutionIn.X),
+		FMath::Clamp(FMath::CeilToInt(FloatBounds.Max.Y) + 1, 0, DestinationResolutionIn.Y));
+
+	// The outer half-pixel shouldn't affect the landscape because it is not part of our official coverage area.
+	EdgeUVDeadBorderOut = FVector2f::Zero();
+	if (SourceResolutionIn.X * SourceResolutionIn.Y != 0)
+	{
+		EdgeUVDeadBorderOut = FVector2f(0.5 / SourceResolutionIn.X, 0.5 / SourceResolutionIn.Y);
+	}
+
+	FVector3d ComponentScale = PatchToWorldOut.GetScale3D();
+	FalloffWorldMarginOut = Falloff / FMath::Min(ComponentScale.X, ComponentScale.Y);
+}
+
+void ULandscapeTexturePatch::GetHeightShaderParams(
+	const FIntPoint& SourceResolutionIn, const FIntPoint& DestinationResolutionIn,
+	UE::Landscape::FApplyLandscapeTextureHeightPatchPS::FParameters& ParamsOut,
+	FIntRect& DestinationBoundsOut) const
+{
+	using namespace UE::Landscape;
+
+	FTransform PatchToWorld;
+	GetCommonShaderParams(SourceResolutionIn, DestinationResolutionIn,
+		PatchToWorld, ParamsOut.InPatchWorldDimensions, ParamsOut.InHeightmapToPatch, 
+		DestinationBoundsOut, ParamsOut.InEdgeUVDeadBorder, ParamsOut.InFalloffWorldMargin);
+
+	FVector3d ComponentScale = PatchToWorld.GetScale3D();
+	double LandscapeHeightScale = Landscape.IsValid() ? Landscape->GetTransform().GetScale3D().Z : 1;
+	LandscapeHeightScale = LandscapeHeightScale == 0 ? 1 : LandscapeHeightScale;
+
+	bool bNativeEncoding = HeightSourceMode == ELandscapeTexturePatchSourceMode::InternalTexture
+		|| HeightEncoding == ELandscapeTextureHeightPatchEncoding::NativePackedHeight;
+
+	// To get height scale in heightmap coordinates, we have to undo the scaling that happens to map the 16bit int to [-256, 256), and undo
+	// the landscape actor scale.
+	ParamsOut.InHeightScale = bNativeEncoding ? 1
+		: LANDSCAPE_INV_ZSCALE * HeightEncodingSettings.WorldSpaceEncodingScale / LandscapeHeightScale;
+	if (bApplyComponentZScale)
+	{
+		ParamsOut.InHeightScale *= ComponentScale.Z;
+	}
+
+	ParamsOut.InZeroInEncoding = bNativeEncoding ? LandscapeDataAccess::MidValue : HeightEncodingSettings.ZeroInEncoding;
+
+	ParamsOut.InHeightOffset = 0;
+	switch (ZeroHeightMeaning)
+	{
+	case ELandscapeTextureHeightPatchZeroHeightMeaning::LandscapeZ:
+		break; // no offset necessary
+	case ELandscapeTextureHeightPatchZeroHeightMeaning::PatchZ:
+	{
+		FVector3d PatchOriginInHeightmapCoords = PatchManager->GetHeightmapCoordsToWorld().InverseTransformPosition(PatchToWorld.GetTranslation());
+		ParamsOut.InHeightOffset = PatchOriginInHeightmapCoords.Z - LandscapeDataAccess::MidValue;
+		break;
+	}
+	case ELandscapeTextureHeightPatchZeroHeightMeaning::WorldZero:
+	{
+		FVector3d WorldOriginInHeightmapCoords = PatchManager->GetHeightmapCoordsToWorld().InverseTransformPosition(FVector::ZeroVector);
+		ParamsOut.InHeightOffset = WorldOriginInHeightmapCoords.Z - LandscapeDataAccess::MidValue;
+		break;
+	}
+	default:
+		ensure(false);
+	}
+
+	ParamsOut.InBlendMode = static_cast<uint32>(BlendMode);
+
+	// Pack our booleans into a bitfield
+	using EShaderFlags = FApplyLandscapeTextureHeightPatchPS::EFlags;
+	EShaderFlags Flags = EShaderFlags::None;
+
+	Flags |= (FalloffMode == ELandscapeTexturePatchFalloffMode::RoundedRectangle) ?
+		EShaderFlags::RectangularFalloff : EShaderFlags::None;
+
+	Flags |= bUseTextureAlphaForHeight ?
+		EShaderFlags::ApplyPatchAlpha : EShaderFlags::None;
+
+	Flags |= bNativeEncoding ?
+		EShaderFlags::InputIsPackedHeight : EShaderFlags::None;
+
+	ParamsOut.InFlags = static_cast<uint8>(Flags);
+}
+
+void ULandscapeTexturePatch::GetWeightShaderParams(const FIntPoint& SourceResolutionIn, 
+	const FIntPoint& DestinationResolutionIn, const ULandscapeWeightPatchTextureInfo* WeightPatchInfo, 
+	UE::Landscape::FApplyLandscapeTextureWeightPatchPS::FParameters& ParamsOut, 
+	FIntRect& DestinationBoundsOut) const
+{
+	using namespace UE::Landscape;
+
+	FTransform PatchToWorld;
+	GetCommonShaderParams(SourceResolutionIn, DestinationResolutionIn,
+		PatchToWorld, ParamsOut.InPatchWorldDimensions, ParamsOut.InWeightmapToPatch,
+		DestinationBoundsOut, ParamsOut.InEdgeUVDeadBorder, ParamsOut.InFalloffWorldMargin);
+
+	// Use the override blend mode if present, otherwise fall back to more general blend mode.
+	ParamsOut.InBlendMode = static_cast<uint32>(WeightPatchInfo->bOverrideBlendMode ? WeightPatchInfo->OverrideBlendMode : BlendMode);
+
+	// Pack our booleans into a bitfield
+	using EShaderFlags = FApplyLandscapeTextureHeightPatchPS::EFlags;
+	EShaderFlags Flags = EShaderFlags::None;
+
+	Flags |= (FalloffMode == ELandscapeTexturePatchFalloffMode::RoundedRectangle) ?
+		EShaderFlags::RectangularFalloff : EShaderFlags::None;
+
+	Flags |= WeightPatchInfo->bUseAlphaChannel ?
+		EShaderFlags::ApplyPatchAlpha : EShaderFlags::None;
+
+	ParamsOut.InFlags = static_cast<uint8>(Flags);
+}
+
+FLandscapeHeightPatchConvertToNativeParams ULandscapeTexturePatch::GetHeightConversionParams() const
+{
+	// When doing conversions, we bake into a height in the same way that we do when applying the patch.
+
+	FLandscapeHeightPatchConvertToNativeParams ConversionParams;
+	ConversionParams.ZeroInEncoding = HeightEncodingSettings.ZeroInEncoding;
+
+	double LandscapeHeightScale = Landscape.IsValid() ? Landscape->GetTransform().GetScale3D().Z : 1;
+	LandscapeHeightScale = LandscapeHeightScale == 0 ? 1 : LandscapeHeightScale;
+	ConversionParams.HeightScale = HeightEncodingSettings.WorldSpaceEncodingScale * LANDSCAPE_INV_ZSCALE / LandscapeHeightScale;
+
+	// TODO: We can choose whether we want to bake in the height offset if it exists. Doing so will handle
+	// some edge cases where the value stored in the patch is outside the range storeable in the native format
+	// normally, but within the range of the landscape due to the patch being far above/below the landscape to
+	// compensate. However, while this is good for conversions for the purposes of serialization, it's not good
+	// for conversions for the purposes of source mode change, so we would need to do things slightly differently
+	// in the two cases. For now, we'll just not bother with that (unlikely?) edge case.
+	ConversionParams.HeightOffset = 0;
+
+	return ConversionParams;
+}
+
+#endif // WITH_EDITOR
+
+void ULandscapeTexturePatch::ReinitializeHeight()
+{
+#if WITH_EDITOR
+	if (!Landscape.IsValid() || !PatchManager.IsValid())
+	{
+		UE_LOG(LogLandscapePatch, Warning, TEXT("ULandscapeTexturePatch::Reinitialize: No associated landscape to initialize from."));
+		return;
+	}
+
+	FVector2D DesiredResolution(FMath::Max(1, InitTextureSizeX), FMath::Max(1, InitTextureSizeY));
+	if (bBaseResolutionOffLandscape)
+	{
+		GetInitResolutionFromLandscape(ResolutionMultiplier, DesiredResolution);
+	}
+	SetResolution(DesiredResolution);
+
+	bReinitializeHeightOnNextRender = true;
+	RequestLandscapeUpdate();
+
+#endif // WITH_EDITOR
+}
+
+void ULandscapeTexturePatch::ReinitializeWeights()
+{
+#if WITH_EDITOR
+	if (!Landscape.IsValid() || !PatchManager.IsValid())
+	{
+		UE_LOG(LogLandscapePatch, Warning, TEXT("ULandscapeTexturePatch::Reinitialize: No associated landscape to initialize from."));
+		return;
+	}
+
+	FVector2D DesiredResolution(FMath::Max(1, InitTextureSizeX), FMath::Max(1, InitTextureSizeY));
+	if (bBaseResolutionOffLandscape)
+	{
+		GetInitResolutionFromLandscape(ResolutionMultiplier, DesiredResolution);
+	}
+	SetResolution(DesiredResolution);
+
+	const FName VisibilityLayerName = Landscape->VisibilityLayer ? Landscape->VisibilityLayer->LayerName : NAME_None;
+
+	ULandscapeInfo* Info = Landscape->GetLandscapeInfo();
+	if (Info)
+	{
+		for (const FLandscapeInfoLayerSettings& InfoLayerSettings : Info->Layers)
+		{
+			if (!InfoLayerSettings.LayerInfoObj)
+			{
+				continue;
+			}
+			FName WeightmapLayerName = InfoLayerSettings.GetLayerName();
+			if (!ensure(WeightmapLayerName != NAME_None) 
+				// TODO: We are currently unable to edit the visibility layer, but we get it passed in. Skip it so
+				// that we don't get the users' hopes up.
+				|| WeightmapLayerName == VisibilityLayerName)
+			{
+				continue;
+			}
+
+			TArray<TObjectPtr<ULandscapeWeightPatchTextureInfo>> FoundPatches = WeightPatches.FilterByPredicate(
+				[&WeightmapLayerName](const TObjectPtr<ULandscapeWeightPatchTextureInfo>& PatchInfo) { return PatchInfo->WeightmapLayerName == WeightmapLayerName; });
+
+			if (FoundPatches.IsEmpty())
+			{
+				AddWeightPatch(WeightmapLayerName, ELandscapeTexturePatchSourceMode::InternalTexture, false);
+				WeightPatches.Last()->bReinitializeOnNextRender = true;
+			}
+			else
+			{
+				for (TObjectPtr<ULandscapeWeightPatchTextureInfo> PatchInfo : FoundPatches)
+				{
+					PatchInfo->bReinitializeOnNextRender = true;
+				}
+			}
+		}
+		RequestLandscapeUpdate();
+	}
+
+#endif // WITH_EDITOR
+}
+
+#if WITH_EDITOR
+void ULandscapeTexturePatch::ReinitializeHeight(UTextureRenderTarget2D* InCombinedResult)
+{
+	if (HeightSourceMode == ELandscapeTexturePatchSourceMode::TextureAsset)
+	{
+		UE_LOG(LogLandscapePatch, Warning, TEXT("ULandscapeTexturePatch: Cannot reinitialize height patch when source mode is an external texture."));
+		return;
+	}
+
+	if (HeightSourceMode == ELandscapeTexturePatchSourceMode::None)
+	{
+		SetHeightSourceMode(ELandscapeTexturePatchSourceMode::InternalTexture);
+	}
+	else if (IsValid(HeightInternalData))
+	{
+		if (HeightSourceMode == ELandscapeTexturePatchSourceMode::InternalTexture && IsValid(HeightInternalData->GetInternalTexture()))
+		{
+			HeightInternalData->GetInternalTexture()->Modify();
+		}
+		else if (HeightSourceMode == ELandscapeTexturePatchSourceMode::TextureBackedRenderTarget && IsValid(HeightInternalData->GetRenderTarget()))
+		{
+			HeightInternalData->GetRenderTarget()->Modify();
+		}
+	}
+
+	if (!ensure(IsValid(HeightInternalData)))
+	{
+		return;
+	}
+
+	// The way we're going to do it is that we'll copy the packed values directly to a temporary render target, offset 
+	// them if needed (to undo whatever offsetting will happen during application), and store the result directly in the
+	// internal texture. Then we'll update the actual associated render target from the internal texture (if needed) so
+	// that unpacking and height format conversion happens the same way as everywhere else.
+
+	UTextureRenderTarget2D* TemporaryNativeHeightCopy = NewObject<UTextureRenderTarget2D>(this);
+	TemporaryNativeHeightCopy->RenderTargetFormat = ETextureRenderTargetFormat::RTF_RGBA8;
+	TemporaryNativeHeightCopy->InitAutoFormat(ResolutionX, ResolutionY);
+	TemporaryNativeHeightCopy->UpdateResourceImmediate(true);
+	
+	// If ZeroHeightMeaning is not landscape Z, then we're going to be applying an offset to our data when
+	// applying it to landscape, which means we'll need to apply the inverse offset when initializing here
+	// so that we get the same landscape back.
+	double OffsetToApply = 0;
+	if (ZeroHeightMeaning != ELandscapeTextureHeightPatchZeroHeightMeaning::LandscapeZ)
+	{
+		FTransform LandscapeHeightmapToWorld = PatchManager->GetHeightmapCoordsToWorld();
+		double ZeroHeight = 0;
+		if (ZeroHeightMeaning == ELandscapeTextureHeightPatchZeroHeightMeaning::PatchZ)
+		{
+			ZeroHeight = LandscapeHeightmapToWorld.InverseTransformPosition(GetComponentTransform().GetTranslation()).Z;
+		}
+		else if (ZeroHeightMeaning == ELandscapeTextureHeightPatchZeroHeightMeaning::WorldZero)
+		{
+			ZeroHeight = LandscapeHeightmapToWorld.InverseTransformPosition(FVector::ZeroVector).Z;
+		}
+		OffsetToApply = LandscapeDataAccess::MidValue - ZeroHeight;
+	}
+
+	FMatrix44f PatchToSource = GetPatchToHeightmapUVs(TemporaryNativeHeightCopy->SizeX, TemporaryNativeHeightCopy->SizeY, InCombinedResult->SizeX, InCombinedResult->SizeY);
+
+	ENQUEUE_RENDER_COMMAND(LandscapeTexturePatchReinitializeHeight)(
+		[Source = InCombinedResult->GetResource(), Destination = TemporaryNativeHeightCopy->GetResource(),
+		&PatchToSource, OffsetToApply](FRHICommandListImmediate& RHICmdList)
+	{
+		using namespace UE::Landscape;
+
+		FRDGBuilder GraphBuilder(RHICmdList, RDG_EVENT_NAME("LandscapeTexturePatchReinitializeHeight"));
+
+		FReinitializeLandscapePatchPS::FParameters* HeightmapResalmpleParams = GraphBuilder.AllocParameters<FReinitializeLandscapePatchPS::FParameters>();
+
+		FRDGTextureRef HeightmapSource = GraphBuilder.RegisterExternalTexture(CreateRenderTarget(Source->GetTexture2DRHI(), TEXT("ReinitializationSource")));
+		FRDGTextureSRVRef SourceSRV = GraphBuilder.CreateSRV(FRDGTextureSRVDesc::CreateForMipLevel(HeightmapSource, 0));
+		HeightmapResalmpleParams->InSource = SourceSRV;
+		HeightmapResalmpleParams->InSourceSampler = TStaticSamplerState<SF_Bilinear, AM_Clamp, AM_Clamp>::GetRHI();
+		HeightmapResalmpleParams->InPatchToSource = PatchToSource;
+
+		FRDGTextureRef DestinationTexture = GraphBuilder.RegisterExternalTexture(CreateRenderTarget(Destination->GetTexture2DRHI(), TEXT("ReinitializationDestination")));
+
+		if (OffsetToApply != 0)
+		{
+			FRDGTextureRef TemporaryDestination = GraphBuilder.CreateTexture(DestinationTexture->Desc, TEXT("LandscapeTextureHeightPatchInputCopy"));
+			HeightmapResalmpleParams->RenderTargets[0] = FRenderTargetBinding(TemporaryDestination, ERenderTargetLoadAction::ENoAction, /*InMipIndex = */0);
+
+			FReinitializeLandscapePatchPS::AddToRenderGraph(GraphBuilder, HeightmapResalmpleParams);
+
+			FOffsetHeightmapPS::FParameters* OffsetParams = GraphBuilder.AllocParameters<FOffsetHeightmapPS::FParameters>();
+
+			FRDGTextureSRVRef InputSRV = GraphBuilder.CreateSRV(FRDGTextureSRVDesc::CreateForMipLevel(TemporaryDestination, 0));
+			OffsetParams->InHeightmap = InputSRV;
+			OffsetParams->InHeightOffset = OffsetToApply;
+			OffsetParams->RenderTargets[0] = FRenderTargetBinding(DestinationTexture, ERenderTargetLoadAction::ENoAction, /*InMipIndex = */0);
+
+			FOffsetHeightmapPS::AddToRenderGraph(GraphBuilder, OffsetParams);
+		}
+		else
+		{
+			HeightmapResalmpleParams->RenderTargets[0] = FRenderTargetBinding(DestinationTexture, ERenderTargetLoadAction::ENoAction, /*InMipIndex = */0);
+			FReinitializeLandscapePatchPS::AddToRenderGraph(GraphBuilder, HeightmapResalmpleParams);
+		}
+
+		GraphBuilder.Execute();
+	});
+
+	TemporaryNativeHeightCopy->UpdateTexture2D(HeightInternalData->GetInternalTexture(), ETextureSourceFormat::TSF_BGRA8);
+	HeightInternalData->GetInternalTexture()->UpdateResource();
+
+	if (IsValid(HeightInternalData->GetRenderTarget()))
+	{
+		HeightInternalData->CopyBackFromInternalTexture();
+	}
+}
+
+void ULandscapeTexturePatch::ReinitializeWeightPatch(ULandscapeWeightPatchTextureInfo* PatchInfo, UTextureRenderTarget2D* InCombinedResult)
+{
+	using namespace LandscapeTexturePatchLocals;
+
+	if (!ensure(IsValid(PatchInfo) && IsValid(InCombinedResult)))
+	{
+		return;
+	}
+
+	if (PatchInfo->SourceMode == ELandscapeTexturePatchSourceMode::TextureAsset)
+	{
+		const FString LayerNameString = PatchInfo->WeightmapLayerName.ToString();
+		UE_LOG(LogLandscapePatch, Warning, TEXT("ULandscapeTexturePatch: Cannot initialize weight layer %s because source mode is an external texture."), *LayerNameString);
+		return;
+	}
+
+	if (PatchInfo->SourceMode == ELandscapeTexturePatchSourceMode::None)
+	{
+		PatchInfo->SetSourceMode(ELandscapeTexturePatchSourceMode::InternalTexture);
+	}
+	else if (IsValid(PatchInfo->InternalData))
+	{
+		if (PatchInfo->SourceMode == ELandscapeTexturePatchSourceMode::InternalTexture && IsValid(PatchInfo->InternalData->GetInternalTexture()))
+		{
+			PatchInfo->InternalData->GetInternalTexture()->Modify();
+		}
+		else if (PatchInfo->SourceMode == ELandscapeTexturePatchSourceMode::TextureBackedRenderTarget && IsValid(PatchInfo->InternalData->GetRenderTarget()))
+		{
+			PatchInfo->InternalData->GetRenderTarget()->Modify();
+		}
+	}
+	
+	if (!ensure(PatchInfo->InternalData))
+	{
+		return;
+	}
+
+	// We're going to copy directly to the associated render target. Make sure there is one for us to copy to.
+	PatchInfo->InternalData->SetUseInternalTextureOnly(false, false);
+	UTextureRenderTarget2D* RenderTarget = PatchInfo->InternalData->GetRenderTarget();
+	if (!ensure(IsValid(RenderTarget)))
+	{
+		return;
+	}
+
+	FMatrix44f PatchToSource = GetPatchToHeightmapUVs(RenderTarget->SizeX, RenderTarget->SizeY, InCombinedResult->SizeX, InCombinedResult->SizeY);
+
+	ENQUEUE_RENDER_COMMAND(LandscapeTexturePatchReinitializeHeight)(
+		[Source = InCombinedResult->GetResource(), Destination = RenderTarget->GetResource(), &PatchToSource](FRHICommandListImmediate& RHICmdList)
+	{
+		using namespace UE::Landscape;
+
+		FRDGBuilder GraphBuilder(RHICmdList, RDG_EVENT_NAME("LandscapeTexturePatchReinitializeHeight"));
+
+		FReinitializeLandscapePatchPS::FParameters* ShaderParams = GraphBuilder.AllocParameters<FReinitializeLandscapePatchPS::FParameters>();
+
+		FRDGTextureRef SourceTexture = GraphBuilder.RegisterExternalTexture(CreateRenderTarget(Source->GetTexture2DRHI(), TEXT("ReinitializationSource")));
+		FRDGTextureSRVRef SourceSRV = GraphBuilder.CreateSRV(FRDGTextureSRVDesc::CreateForMipLevel(SourceTexture, 0));
+		ShaderParams->InSource = SourceSRV;
+		ShaderParams->InSourceSampler = TStaticSamplerState<SF_Bilinear, AM_Clamp, AM_Clamp>::GetRHI();
+
+		ShaderParams->InPatchToSource = PatchToSource;
+
+		FRDGTextureRef DestinationTexture = GraphBuilder.RegisterExternalTexture(CreateRenderTarget(Destination->GetTexture2DRHI(), TEXT("ReinitializationDestination")));
+		ShaderParams->RenderTargets[0] = FRenderTargetBinding(DestinationTexture, ERenderTargetLoadAction::ENoAction, /*InMipIndex = */0);
+		FReinitializeLandscapePatchPS::AddToRenderGraph(GraphBuilder, ShaderParams);
+
+		GraphBuilder.Execute();
+	});
+
+	PatchInfo->InternalData->SetUseInternalTextureOnly(PatchInfo->SourceMode == ELandscapeTexturePatchSourceMode::InternalTexture, true);
+}
+
+FMatrix44f ULandscapeTexturePatch::GetPatchToHeightmapUVs(int32 PatchSizeX, int32 PatchSizeY, int32 HeightmapSizeX, int32 HeightmapSizeY) const
+{
+	FVector2D FullPatchDimensions = GetFullUnscaledWorldSize();
+
+	FTransform PatchPixelToPatchLocal(FQuat4d::Identity, FVector3d(-FullPatchDimensions.X / 2, -FullPatchDimensions.Y / 2, 0),
+		FVector3d(FullPatchDimensions.X / PatchSizeX, FullPatchDimensions.Y / PatchSizeY, 1));
+
+	FTransform PatchToWorld = GetPatchToWorldTransform();
+
+	FTransform LandscapeHeightmapToWorld = PatchManager->GetHeightmapCoordsToWorld();
+	FTransform LandscapeUVToWorld = LandscapeHeightmapToWorld;
+	LandscapeUVToWorld.MultiplyScale3D(FVector3d(HeightmapSizeX, HeightmapSizeY, 1));
+
+	// In unreal, matrix composition is done by multiplying the subsequent ones on the right, and the result
+	// is transpose of what our shader will expect (because unreal right multiplies vectors by matrices).
+	FMatrix44d PatchToLandscapeUVTransposed = PatchPixelToPatchLocal.ToMatrixWithScale() * PatchToWorld.ToMatrixWithScale()
+		* LandscapeUVToWorld.ToInverseMatrixWithScale();
+	return (FMatrix44f)PatchToLandscapeUVTransposed.GetTransposed();
+}
+
+bool ULandscapeTexturePatch::IsAffectingWeightmapLayer(const FName& InLayerName) const
+{
+	if (!IsEnabled())
+	{
+		return false;
+	}
+
+	for (TObjectPtr<ULandscapeWeightPatchTextureInfo> WeightPatch : WeightPatches)
+	{
+		if (WeightPatch->WeightmapLayerName == InLayerName)
+		{
+			return true;
+		}
+	}
+	return false;
+}
+#endif // WITH_EDITOR
+
+void ULandscapeTexturePatch::SnapToLandscape()
+{
+#if WITH_EDITOR
+	if (!Landscape.IsValid())
+	{
+		return;
+	}
+
+	Modify();
+
+	FTransform LandscapeTransform = Landscape->GetTransform();
+	FTransform PatchTransform = GetComponentTransform();
+
+	FQuat LandscapeRotation = LandscapeTransform.GetRotation();
+	FQuat PatchRotation = PatchTransform.GetRotation();
+
+	// Get rotation of patch relative to landscape
+	FQuat PatchRotationRelativeLandscape = LandscapeRotation.Inverse() * PatchRotation;
+
+	// Get component of that relative rotation that is around the landscape Z axis.
+	double RadiansAroundZ = PatchRotationRelativeLandscape.GetTwistAngle(FVector::ZAxisVector);
+
+	// Round that rotation to nearest 90 degree increment
+	int32 Num90DegreeRotations = FMath::RoundToDouble(RadiansAroundZ / FMathd::HalfPi);
+	double NewRadiansAroundZ = Num90DegreeRotations * FMathd::HalfPi;
+
+	// Now adjust the patch transform.
+	FQuat NewPatchRotation = FQuat(FVector::ZAxisVector, NewRadiansAroundZ) * LandscapeRotation;
+	SetWorldRotation(NewPatchRotation);
+
+	// Once we have the rotation adjusted, we need to adjust the patch size and positioning.
+	// However don't bother if either the patch or landscape scale is 0. We might still be able
+	// to align in one of the axes in such a case, but it is not worth the code complexity for
+	// a broken use case.
+	FVector LandscapeScale = Landscape->GetTransform().GetScale3D();
+	FVector PatchScale = GetComponentTransform().GetScale3D();
+	if (LandscapeScale.X == 0 || LandscapeScale.Y == 0)
+	{
+		UE_LOG(LogLandscapePatch, Warning, TEXT("ULandscapeTexturePatch::SnapToLandscape: Landscape target "
+			"for patch had a zero scale in one of the dimensions. Skipping aligning position."));
+		return;
+	}
+	if (PatchScale.X == 0 || PatchScale.Y == 0)
+	{
+		UE_LOG(LogLandscapePatch, Warning, TEXT("ULandscapeTexturePatch::SnapToLandscape: Patch "
+			"had a zero scale in one of the dimensions. Skipping aligning position."));
+		return;
+	}
+
+	// Start by adjusting size to be a multiple of landscape quad size.
+	double PatchExtentX = PatchScale.X * UnscaledPatchCoverage.X;
+	double PatchExtentY = PatchScale.Y * UnscaledPatchCoverage.Y;
+	if (Num90DegreeRotations % 2)
+	{
+		// Relative to the landscape, our lenght and width are backwards...
+		Swap(PatchExtentX, PatchExtentY);
+	}
+
+	int32 LandscapeQuadsX = FMath::RoundToInt(PatchExtentX / LandscapeScale.X);
+	int32 LandscapeQuadsY = FMath::RoundToInt(PatchExtentY / LandscapeScale.Y);
+
+	double NewPatchExtentX = LandscapeQuadsX * LandscapeScale.X;
+	double NewPatchExtentY = LandscapeQuadsY * LandscapeScale.Y;
+	if (Num90DegreeRotations % 2)
+	{
+		Swap(NewPatchExtentX, NewPatchExtentY);
+	}
+	UnscaledPatchCoverage = FVector2D(NewPatchExtentX / PatchScale.X, NewPatchExtentY / PatchScale.Y);
+
+	// Now adjust the center of the patch. This gets snapped to either integer or integer + 0.5 increments
+	// in landscape coordinates depending on whether patch length/width is odd or even in landscape coordinates.
+
+	FVector PatchCenterInLandscapeCoordinates = LandscapeTransform.InverseTransformPosition(GetComponentLocation());
+	double NewPatchCenterX, NewPatchCenterY;
+	if (LandscapeQuadsX % 2)
+	{
+		NewPatchCenterX = FMath::RoundToDouble(PatchCenterInLandscapeCoordinates.X + 0.5) - 0.5;
+	}
+	else
+	{
+		NewPatchCenterX = FMath::RoundToDouble(PatchCenterInLandscapeCoordinates.X);
+	}
+	if (LandscapeQuadsY % 2)
+	{
+		NewPatchCenterY = FMath::RoundToDouble(PatchCenterInLandscapeCoordinates.Y + 0.5) - 0.5;
+	}
+	else
+	{
+		NewPatchCenterY = FMath::RoundToDouble(PatchCenterInLandscapeCoordinates.Y);
+	}
+
+	FVector NewCenterInLandscape(NewPatchCenterX, NewPatchCenterY, PatchCenterInLandscapeCoordinates.Z);
+	SetWorldLocation(LandscapeTransform.TransformPosition(NewCenterInLandscape));
+	RequestLandscapeUpdate();
+#endif // WITH_EDITOR
+}
+
+void ULandscapeTexturePatch::SetResolution(FVector2D ResolutionIn)
+{
+	int32 DesiredX = FMath::Max(1, ResolutionIn.X);
+	int32 DesiredY = FMath::Max(1, ResolutionIn.Y);
+
+	if (DesiredX == ResolutionX && DesiredY == ResolutionY)
+	{
+		return;
+	}
+	Modify();
+
+	ResolutionX = DesiredX;
+	ResolutionY = DesiredY;
+	InitTextureSizeX = ResolutionX;
+	InitTextureSizeY = ResolutionY;
+
+	auto ResizePatch = [DesiredX, DesiredY](ELandscapeTexturePatchSourceMode SourceMode, ULandscapeTextureBackedRenderTargetBase* InternalData)
+	{
+		// Deal with height first
+		if (SourceMode == ELandscapeTexturePatchSourceMode::TextureAsset || SourceMode == ELandscapeTexturePatchSourceMode::None)
+		{
+			return;
+		}
+		else
+		{
+			InternalData->SetResolution(DesiredX, DesiredY);
+		}
+	};
+
+	ResizePatch(HeightSourceMode, HeightInternalData);
+	for (TObjectPtr<ULandscapeWeightPatchTextureInfo> WeightPatch : WeightPatches)
+	{
+		ResizePatch(WeightPatch->SourceMode, WeightPatch->InternalData);
+	}
+}
+
+FVector2D ULandscapeTexturePatch::GetFullUnscaledWorldSize() const
+{
+	FVector2D Resolution = GetResolution();
+
+	// UnscaledPatchCoverage is meant to represent the distance between the centers of the extremal pixels.
+	// That distance in pixels is Resolution-1.
+	FVector2D TargetPixelSize(UnscaledPatchCoverage / FMath::Max(Resolution - 1, FVector2D(1, 1)));
+	return TargetPixelSize * Resolution;
+}
+
+FTransform ULandscapeTexturePatch::GetPatchToWorldTransform() const
+{
+	FTransform PatchToWorld = GetComponentTransform();
+
+	if (Landscape.IsValid())
+	{
+		FRotator3d PatchRotator = PatchToWorld.GetRotation().Rotator();
+		FRotator3d LandscapeRotator = Landscape->GetTransform().GetRotation().Rotator();
+		PatchToWorld.SetRotation(FRotator3d(LandscapeRotator.Pitch, PatchRotator.Yaw, LandscapeRotator.Roll).Quaternion());
+	}
+
+	return PatchToWorld;
+}
+
+bool ULandscapeTexturePatch::GetInitResolutionFromLandscape(float ResolutionMultiplierIn, FVector2D& ResolutionOut) const
+{
+	if (!Landscape.IsValid())
+	{
+		return false;
+	}
+
+	ResolutionOut = FVector2D::One();
+
+	FVector LandscapeScale = Landscape->GetTransform().GetScale3D();
+	// We go off of the larger dimension so that our patch works in different rotations.
+	double LandscapeQuadSize = FMath::Max(FMath::Abs(LandscapeScale.X), FMath::Abs(LandscapeScale.Y));
+
+	if (LandscapeQuadSize > 0)
+	{
+		double PatchQuadSize = LandscapeQuadSize;
+		PatchQuadSize /= (ResolutionMultiplierIn > 0 ? ResolutionMultiplierIn : 1);
+
+		FVector PatchScale = GetComponentTransform().GetScale3D();
+		double NumQuadsX = FMath::Abs(UnscaledPatchCoverage.X * PatchScale.X / PatchQuadSize);
+		double NumQuadsY = FMath::Abs(UnscaledPatchCoverage.Y * PatchScale.Y / PatchQuadSize);
+
+		ResolutionOut = FVector2D(
+			FMath::Max(1, FMath::CeilToInt(NumQuadsX) + 1),
+			FMath::Max(1, FMath::CeilToInt(NumQuadsY) + 1)
+		);
+
+		return true;
+	}
+	return false;
+}
+
+#if WITH_EDITOR
+void ULandscapeTexturePatch::PostEditChangeProperty(FPropertyChangedEvent& PropertyChangedEvent)
+{
+	using namespace LandscapeTexturePatchLocals;
+
+	if (PropertyChangedEvent.Property)
+	{
+		if (PropertyChangedEvent.Property->GetFName() == GET_MEMBER_NAME_CHECKED(ULandscapeTexturePatch, DetailPanelHeightSourceMode))
+		{
+			SetHeightSourceMode(DetailPanelHeightSourceMode);
+		}
+		else if (PropertyChangedEvent.Property->GetFName() == GET_MEMBER_NAME_CHECKED(ULandscapeTexturePatch, HeightEncoding))
+		{
+			ResetHeightEncodingMode(HeightEncoding);
+		}
+		else if (PropertyChangedEvent.Property->GetFName() == GET_MEMBER_NAME_CHECKED(ULandscapeTexturePatch, WeightPatches))
+		{
+			if (NumWeightPatches != WeightPatches.Num())
+			{
+				// User must have added/removed a weight patch, so make sure that all the owner pointers are set
+				for (TObjectPtr<ULandscapeWeightPatchTextureInfo>& WeightPatch : WeightPatches)
+				{
+					if (!IsValid(WeightPatch))
+					{
+						// The entries that users add start out as null
+						WeightPatch = NewObject<ULandscapeWeightPatchTextureInfo>(this);
+					}
+					WeightPatch->OwningPatch = this;
+				}
+				NumWeightPatches = WeightPatches.Num();
+			}
+		}
+	}
+
+	Super::PostEditChangeProperty(PropertyChangedEvent);
+}
+
+void ULandscapeWeightPatchTextureInfo::PostEditChangeProperty(FPropertyChangedEvent& PropertyChangedEvent)
+{
+	using namespace LandscapeTexturePatchLocals;
+
+	if (PropertyChangedEvent.Property)
+	{
+		if (PropertyChangedEvent.Property->GetFName() == GET_MEMBER_NAME_CHECKED(ULandscapeWeightPatchTextureInfo, DetailPanelSourceMode)
+			&& DetailPanelSourceMode != SourceMode)
+		{
+			SetSourceMode(DetailPanelSourceMode);
+		}
+	}
+
+	Super::PostEditChangeProperty(PropertyChangedEvent);
+}
+#endif // WITH_EDITOR
+
+void ULandscapeWeightPatchTextureInfo::SetSourceMode(ELandscapeTexturePatchSourceMode NewMode)
+{
+#if WITH_EDITOR
+	using namespace LandscapeTexturePatchLocals;
+
+	if (SourceMode == NewMode)
+	{
+		return;
+	}
+	Modify();
+
+	FVector2D Resolution(1, 1);
+	if (OwningPatch)
+	{
+		Resolution = OwningPatch->GetResolution();
+	}
+
+	TransitionSourceMode<ULandscapeWeightTextureBackedRenderTarget>(SourceMode, NewMode, TextureAsset, InternalData, [&Resolution, this]()
+	{
+		ULandscapeWeightTextureBackedRenderTarget* InternalDataToReturn = NewObject<ULandscapeWeightTextureBackedRenderTarget>(this);
+		InternalDataToReturn->SetResolution(Resolution.X, Resolution.Y);
+		return InternalDataToReturn;
+	});
+
+	SourceMode = NewMode;
+	DetailPanelSourceMode = NewMode;
+#endif // WITH_EDITOR
+}
+
+void ULandscapeTexturePatch::SetHeightSourceMode(ELandscapeTexturePatchSourceMode NewMode)
+{
+#if WITH_EDITOR
+	using namespace LandscapeTexturePatchLocals;
+
+	if (HeightSourceMode == NewMode)
+	{
+		return;
+	}
+	Modify();
+
+	TransitionSourceMode<ULandscapeHeightTextureBackedRenderTarget>(HeightSourceMode, NewMode, HeightTextureAsset, HeightInternalData, [this]()
+	{
+		ULandscapeHeightTextureBackedRenderTarget* InternalDataToReturn = NewObject<ULandscapeHeightTextureBackedRenderTarget>(this);
+		InternalDataToReturn->SetResolution(ResolutionX, ResolutionY);
+		InternalDataToReturn->SetFormat(HeightRenderTargetFormat);
+		InternalDataToReturn->ConversionParams = GetHeightConversionParams();
+
+		return InternalDataToReturn;
+	});
+
+	HeightSourceMode = NewMode;
+	DetailPanelHeightSourceMode = NewMode;
+#endif // WITH_EDITOR
+}
+
+void ULandscapeTexturePatch::SetHeightTextureAsset(UTexture* TextureIn)
+{
+	ensureMsgf(!TextureIn || TextureIn->VirtualTextureStreaming == 0,
+		TEXT("ULandscapeTexturePatch::SetHeightTextureAsset: Virtual textures are not supported."));
+	HeightTextureAsset = TextureIn;
+}
+
+void ULandscapeTexturePatch::ResetHeightEncodingMode(ELandscapeTextureHeightPatchEncoding EncodingMode)
+{
+	Modify();
+	HeightEncoding = EncodingMode;
+	if (EncodingMode == ELandscapeTextureHeightPatchEncoding::ZeroToOne)
+	{
+		HeightEncodingSettings.ZeroInEncoding = 0.5;
+		HeightEncodingSettings.WorldSpaceEncodingScale = 400;
+	}
+	else if (EncodingMode == ELandscapeTextureHeightPatchEncoding::WorldUnits)
+	{
+		HeightEncodingSettings.ZeroInEncoding = 0;
+		HeightEncodingSettings.WorldSpaceEncodingScale = 1;
+	}
+}
+
+void ULandscapeTexturePatch::SetHeightRenderTargetFormat(ETextureRenderTargetFormat Format)
+{
+	if (HeightRenderTargetFormat == Format)
+	{
+		return;
+	}
+
+	Modify();
+	HeightRenderTargetFormat = Format;
+	if (HeightInternalData)
+	{
+		HeightInternalData->SetFormat(HeightRenderTargetFormat);
+	}
+}
+
+
+void ULandscapeTexturePatch::AddWeightPatch(const FName& WeightmapLayerName, ELandscapeTexturePatchSourceMode SourceMode, bool bUseAlphaChannel)
+{
+#if WITH_EDITOR
+	using namespace LandscapeTexturePatchLocals;
+
+	// Try to modify an existing entry instead if possible
+	for (TObjectPtr<ULandscapeWeightPatchTextureInfo> WeightPatch : WeightPatches)
+	{
+		if (WeightPatch->WeightmapLayerName == WeightmapLayerName)
+		{
+			if (WeightPatch->SourceMode != SourceMode)
+			{
+				WeightPatch->SetSourceMode(SourceMode);
+			}
+			if (IsValid(WeightPatch->InternalData))
+			{
+				WeightPatch->InternalData->SetUseAlphaChannel(bUseAlphaChannel);
+			}
+			return;
+		}
+	}
+
+	ULandscapeWeightPatchTextureInfo* NewWeightPatch = NewObject<ULandscapeWeightPatchTextureInfo>(this);
+	NewWeightPatch->WeightmapLayerName = WeightmapLayerName;
+	NewWeightPatch->SourceMode = SourceMode;
+	NewWeightPatch->DetailPanelSourceMode = SourceMode;
+	NewWeightPatch->bUseAlphaChannel = bUseAlphaChannel;
+	NewWeightPatch->OwningPatch = this;
+
+	if (NewWeightPatch->SourceMode == ELandscapeTexturePatchSourceMode::InternalTexture
+		|| NewWeightPatch->SourceMode == ELandscapeTexturePatchSourceMode::TextureBackedRenderTarget)
+	{
+		NewWeightPatch->InternalData = NewObject<ULandscapeWeightTextureBackedRenderTarget>(NewWeightPatch);
+		NewWeightPatch->InternalData->SetResolution(ResolutionX, ResolutionY);
+		NewWeightPatch->InternalData->SetUseAlphaChannel(bUseAlphaChannel);
+		NewWeightPatch->InternalData->Initialize();
+	}
+
+	WeightPatches.Add(NewWeightPatch);
+	++NumWeightPatches;
+#endif // WITH_EDITOR
+}
+
+void ULandscapeTexturePatch::RemoveWeightPatch(const FName& InWeightmapLayerName)
+{
+	NumWeightPatches -= WeightPatches.RemoveAll([InWeightmapLayerName](const TObjectPtr<ULandscapeWeightPatchTextureInfo>& PatchInfo)
+	{ 
+		return PatchInfo->WeightmapLayerName == InWeightmapLayerName; 
+	});
+}
+
+void ULandscapeTexturePatch::RemoveAllWeightPatches()
+{
+	WeightPatches.Reset();
+	NumWeightPatches = 0;
+}
+
+void ULandscapeTexturePatch::DisableAllWeightPatches()
+{
+	for (TObjectPtr<ULandscapeWeightPatchTextureInfo> WeightPatch : WeightPatches)
+	{
+		WeightPatch->SetSourceMode(ELandscapeTexturePatchSourceMode::None);
+	}
+}
+
+TArray<FName> ULandscapeTexturePatch::GetAllWeightPatchLayerNames()
+{
+	TArray<FName> Names;
+	for (TObjectPtr<ULandscapeWeightPatchTextureInfo> WeightPatch : WeightPatches)
+	{
+		if (WeightPatch->WeightmapLayerName != NAME_None)
+		{
+			Names.AddUnique(WeightPatch->WeightmapLayerName);
+		}
+	}
+
+	return Names;
+}
+
+void ULandscapeTexturePatch::SetUseAlphaChannelForWeightPatch(const FName& InWeightmapLayerName, bool bUseAlphaChannel)
+{
+	for (TObjectPtr<ULandscapeWeightPatchTextureInfo> WeightPatch : WeightPatches)
+	{
+		if (WeightPatch->WeightmapLayerName == InWeightmapLayerName)
+		{
+			WeightPatch->bUseAlphaChannel = bUseAlphaChannel;
+			if (WeightPatch->InternalData)
+			{
+				WeightPatch->InternalData->SetUseAlphaChannel(bUseAlphaChannel);
+			}
+			return;
+		}
+	}
+	const FString LayerNameString = InWeightmapLayerName.ToString();
+	UE_LOG(LogLandscapePatch, Warning, TEXT("ULandscapeTexturePatch::SetUseAlphaChannelForWeightPatch: Unable to find data for weight layer %s"), *LayerNameString);
+}
+
+void ULandscapeTexturePatch::SetWeightPatchSourceMode(const FName& InWeightmapLayerName, ELandscapeTexturePatchSourceMode NewMode)
+{
+#if WITH_EDITOR
+	using namespace LandscapeTexturePatchLocals;
+
+	for (TObjectPtr<ULandscapeWeightPatchTextureInfo> WeightPatch : WeightPatches)
+	{
+		if (WeightPatch->WeightmapLayerName == InWeightmapLayerName)
+		{
+			if (WeightPatch->SourceMode != NewMode)
+			{
+				WeightPatch->SetSourceMode(NewMode);
+			}
+			return;
+		}
+	}
+	const FString LayerNameString = InWeightmapLayerName.ToString();
+	UE_LOG(LogLandscapePatch, Warning, TEXT("ULandscapeTexturePatch::SetWeightPatchSourceMode: Unable to find data for weight layer %s"), *LayerNameString);
+#endif // WITH_EDITOR
+}
+
+ELandscapeTexturePatchSourceMode ULandscapeTexturePatch::GetWeightPatchSourceMode(const FName& InWeightmapLayerName)
+{
+	for (TObjectPtr<ULandscapeWeightPatchTextureInfo> WeightPatch : WeightPatches)
+	{
+		if (WeightPatch->WeightmapLayerName == InWeightmapLayerName)
+		{
+			return WeightPatch->SourceMode;
+		}
+	}
+	return ELandscapeTexturePatchSourceMode::None;
+}
+
+UTextureRenderTarget2D* ULandscapeTexturePatch::GetWeightPatchRenderTarget(const FName& InWeightmapLayerName)
+{
+	for (TObjectPtr<ULandscapeWeightPatchTextureInfo> WeightPatch : WeightPatches)
+	{
+		if (WeightPatch->WeightmapLayerName == InWeightmapLayerName)
+		{
+			return WeightPatch->InternalData ? WeightPatch->InternalData->GetRenderTarget() : nullptr;
+		}
+	}
+	return nullptr;
+}
+
+UTexture2D* ULandscapeTexturePatch::GetWeightPatchInternalTexture(const FName& InWeightmapLayerName)
+{
+	for (TObjectPtr<ULandscapeWeightPatchTextureInfo> WeightPatch : WeightPatches)
+	{
+		if (WeightPatch->WeightmapLayerName == InWeightmapLayerName)
+		{
+			return WeightPatch->InternalData ? WeightPatch->InternalData->GetInternalTexture() : nullptr;
+		}
+	}
+	return nullptr;
+}
+
+void ULandscapeTexturePatch::SetWeightPatchTextureAsset(const FName& InWeightmapLayerName, UTexture* TextureIn)
+{
+	if (!ensureMsgf(!TextureIn || TextureIn->VirtualTextureStreaming == 0,
+		TEXT("ULandscapeTexturePatch::SetWeightPatchTextureAsset: Virtual textures are not supported.")))
+	{
+		return;
+	}
+
+	for (TObjectPtr<ULandscapeWeightPatchTextureInfo> WeightPatch : WeightPatches)
+	{
+		if (WeightPatch->WeightmapLayerName == InWeightmapLayerName)
+		{
+			WeightPatch->TextureAsset = TextureIn;
+			return;
+		}
+	}
+
+	const FString LayerNameString = InWeightmapLayerName.ToString();
+	UE_LOG(LogLandscapePatch, Warning, TEXT("ULandscapeTexturePatch::SetWeightPatchTextureAsset: Unable to find data for weight layer %s"), *LayerNameString);
+}
+
+void ULandscapeTexturePatch::SetWeightPatchBlendModeOverride(const FName& InWeightmapLayerName, ELandscapeTexturePatchBlendMode BlendModeIn)
+{
+	for (TObjectPtr<ULandscapeWeightPatchTextureInfo> WeightPatch : WeightPatches)
+	{
+		if (WeightPatch->WeightmapLayerName == InWeightmapLayerName)
+		{
+			WeightPatch->OverrideBlendMode = BlendModeIn;
+			WeightPatch->bOverrideBlendMode = true;
+			return;
+		}
+	}
+}
+
+void ULandscapeTexturePatch::ClearWeightPatchBlendModeOverride(const FName& InWeightmapLayerName)
+{
+	for (TObjectPtr<ULandscapeWeightPatchTextureInfo> WeightPatch : WeightPatches)
+	{
+		if (WeightPatch->WeightmapLayerName == InWeightmapLayerName)
+		{
+			WeightPatch->bOverrideBlendMode = false;
+			return;
+		}
+	}
+}
diff --git a/Engine/Plugins/Experimental/LandscapePatch/Source/LandscapePatch/Private/LandscapeTexturePatchBase.cpp b/Engine/Plugins/Experimental/LandscapePatch/Source/LandscapePatch/Private/LandscapeTexturePatchBase.cpp
--- a/Engine/Plugins/Experimental/LandscapePatch/Source/LandscapePatch/Private/LandscapeTexturePatchBase.cpp
+++ b/Engine/Plugins/Experimental/LandscapePatch/Source/LandscapePatch/Private/LandscapeTexturePatchBase.cpp
@@ -32,6 +32,8 @@ bool ULandscapeTexturePatchBase::GetTextureResolution(FVector2D& SizeOut) const
 {
 	switch (SourceMode)
 	{
+	case ELandscapeTexturePatchSourceMode::None:
+		return false;
 	case ELandscapeTexturePatchSourceMode::InternalTexture:
 		if (IsValid(InternalTexture))
 		{
diff --git a/Engine/Plugins/Experimental/LandscapePatch/Source/LandscapePatch/Private/LandscapeTexturePatchPS.cpp b/Engine/Plugins/Experimental/LandscapePatch/Source/LandscapePatch/Private/LandscapeTexturePatchPS.cpp
new file mode 100644
--- /dev/null
+++ b/Engine/Plugins/Experimental/LandscapePatch/Source/LandscapePatch/Private/LandscapeTexturePatchPS.cpp
@@ -0,0 +1,241 @@
+// Copyright Epic Games, Inc. All Rights Reserved.
+
+#include "LandscapeTexturePatchPS.h"
+
+#include "LandscapeUtils.h"
+#include "PixelShaderUtils.h"
+
+namespace UE::Landscape
+{
+
+bool FApplyLandscapeTextureHeightPatchPS::ShouldCompilePermutation(const FGlobalShaderPermutationParameters& Parameters)
+{
+	return DoesPlatformSupportEditLayers(Parameters.Platform);
+}
+
+void FApplyLandscapeTextureHeightPatchPS::ModifyCompilationEnvironment(const FGlobalShaderPermutationParameters& Parameters, FShaderCompilerEnvironment& OutEnvironment)
+{
+	OutEnvironment.SetDefine(TEXT("APPLY_HEIGHT_PATCH"), 1);
+
+	// Make our flag choices match in the shader.
+	OutEnvironment.SetDefine(TEXT("RECTANGULAR_FALLOFF_FLAG"), static_cast<uint8>(FApplyLandscapeTextureHeightPatchPS::EFlags::RectangularFalloff));
+	OutEnvironment.SetDefine(TEXT("APPLY_PATCH_ALPHA_FLAG"), static_cast<uint8>(FApplyLandscapeTextureHeightPatchPS::EFlags::ApplyPatchAlpha));
+	OutEnvironment.SetDefine(TEXT("INPUT_IS_PACKED_HEIGHT_FLAG"), static_cast<uint8>(FApplyLandscapeTextureHeightPatchPS::EFlags::InputIsPackedHeight));
+
+	OutEnvironment.SetDefine(TEXT("ADDITIVE_MODE"), static_cast<uint8>(FApplyLandscapeTextureHeightPatchPS::EBlendMode::Additive));
+	OutEnvironment.SetDefine(TEXT("ALPHA_BLEND_MODE"), static_cast<uint8>(FApplyLandscapeTextureHeightPatchPS::EBlendMode::AlphaBlend));
+	OutEnvironment.SetDefine(TEXT("MIN_MODE"), static_cast<uint8>(FApplyLandscapeTextureHeightPatchPS::EBlendMode::Min));
+	OutEnvironment.SetDefine(TEXT("MAX_MODE"), static_cast<uint8>(FApplyLandscapeTextureHeightPatchPS::EBlendMode::Max));
+}
+
+void FApplyLandscapeTextureHeightPatchPS::AddToRenderGraph(FRDGBuilder& GraphBuilder, FParameters* InParameters, const FIntRect& DestinationBounds)
+{
+	FGlobalShaderMap* ShaderMap = GetGlobalShaderMap(GMaxRHIFeatureLevel);
+	TShaderMapRef<FApplyLandscapeTextureHeightPatchPS> PixelShader(ShaderMap);
+
+	FPixelShaderUtils::AddFullscreenPass(
+		GraphBuilder,
+		ShaderMap,
+		RDG_EVENT_NAME("LandscapeTextureHeightPatch"),
+		PixelShader,
+		InParameters,
+		DestinationBounds);
+}
+
+bool FOffsetHeightmapPS::ShouldCompilePermutation(const FGlobalShaderPermutationParameters& Parameters)
+{
+	return DoesPlatformSupportEditLayers(Parameters.Platform);
+}
+
+void FOffsetHeightmapPS::ModifyCompilationEnvironment(const FGlobalShaderPermutationParameters& Parameters, FShaderCompilerEnvironment& OutEnvironment)
+{
+	OutEnvironment.SetDefine(TEXT("OFFSET_HEIGHT_PATCH"), 1);
+}
+
+void FOffsetHeightmapPS::AddToRenderGraph(FRDGBuilder& GraphBuilder, FParameters* InParameters)
+{
+	FGlobalShaderMap* ShaderMap = GetGlobalShaderMap(GMaxRHIFeatureLevel);
+	TShaderMapRef<FOffsetHeightmapPS> PixelShader(ShaderMap);
+
+	FIntVector TextureSize = InParameters->InHeightmap->Desc.Texture->Desc.GetSize();
+
+	FPixelShaderUtils::AddFullscreenPass(
+		GraphBuilder,
+		ShaderMap,
+		RDG_EVENT_NAME("OffsetHeightmap"),
+		PixelShader,
+		InParameters,
+		FIntRect(0, 0, TextureSize.X, TextureSize.Y));
+}
+
+bool FSimpleTextureCopyPS::ShouldCompilePermutation(const FGlobalShaderPermutationParameters& Parameters)
+{
+	return DoesPlatformSupportEditLayers(Parameters.Platform);
+}
+
+void FSimpleTextureCopyPS::ModifyCompilationEnvironment(const FGlobalShaderPermutationParameters& Parameters, FShaderCompilerEnvironment& OutEnvironment)
+{
+	OutEnvironment.SetDefine(TEXT("SIMPLE_TEXTURE_COPY"), 1);
+}
+
+void FSimpleTextureCopyPS::AddToRenderGraph(FRDGBuilder& GraphBuilder, FRDGTextureRef SourceTexture, FRDGTextureRef DestinationTexture)
+{
+	FGlobalShaderMap* ShaderMap = GetGlobalShaderMap(GMaxRHIFeatureLevel);
+	TShaderMapRef<FSimpleTextureCopyPS> PixelShader(ShaderMap);
+
+	FRDGTextureSRVRef SourceSRV = GraphBuilder.CreateSRV(FRDGTextureSRVDesc::CreateForMipLevel(SourceTexture, 0));
+
+	FSimpleTextureCopyPS::FParameters* ShaderParams = GraphBuilder.AllocParameters<FSimpleTextureCopyPS::FParameters>();
+	ShaderParams->InSource = GraphBuilder.CreateSRV(FRDGTextureSRVDesc::CreateForMipLevel(SourceTexture, 0));
+	ShaderParams->InSourceSampler = TStaticSamplerState<SF_Bilinear, AM_Clamp, AM_Clamp>::GetRHI();
+	ShaderParams->RenderTargets[0] = FRenderTargetBinding(DestinationTexture, ERenderTargetLoadAction::ENoAction, /*InMipIndex = */0);
+	FIntVector DestinationSize = DestinationTexture->Desc.GetSize();
+	ShaderParams->InDestinationResolution = FVector2f(DestinationSize.X, DestinationSize.Y);
+
+	FPixelShaderUtils::AddFullscreenPass(
+		GraphBuilder,
+		ShaderMap,
+		RDG_EVENT_NAME("SimpleTextureCopy"),
+		PixelShader,
+		ShaderParams,
+		FIntRect(0, 0, DestinationSize.X, DestinationSize.Y));
+}
+
+bool FConvertToNativeLandscapePatchPS::ShouldCompilePermutation(const FGlobalShaderPermutationParameters& Parameters)
+{
+	return DoesPlatformSupportEditLayers(Parameters.Platform);
+}
+
+void FConvertToNativeLandscapePatchPS::ModifyCompilationEnvironment(const FGlobalShaderPermutationParameters& Parameters, FShaderCompilerEnvironment& OutEnvironment)
+{
+	OutEnvironment.SetDefine(TEXT("CONVERT_TO_NATIVE_LANDSCAPE_PATCH"), 1);
+}
+
+void FConvertToNativeLandscapePatchPS::AddToRenderGraph(FRDGBuilder& GraphBuilder, FRDGTextureRef SourceTexture, 
+	FRDGTextureRef DestinationTexture, const FLandscapeHeightPatchConvertToNativeParams& Params)
+{
+	FGlobalShaderMap* ShaderMap = GetGlobalShaderMap(GMaxRHIFeatureLevel);
+	TShaderMapRef<FConvertToNativeLandscapePatchPS> PixelShader(ShaderMap);
+
+	FRDGTextureSRVRef SourceSRV = GraphBuilder.CreateSRV(FRDGTextureSRVDesc::CreateForMipLevel(SourceTexture, 0));
+
+	FConvertToNativeLandscapePatchPS::FParameters* ShaderParams = GraphBuilder.AllocParameters<FConvertToNativeLandscapePatchPS::FParameters>();
+	ShaderParams->InHeightmap = GraphBuilder.CreateSRV(FRDGTextureSRVDesc::CreateForMipLevel(SourceTexture, 0));
+	ShaderParams->InZeroInEncoding = Params.ZeroInEncoding;
+	ShaderParams->InHeightScale = Params.HeightScale;
+	ShaderParams->InHeightOffset = Params.HeightOffset;
+	ShaderParams->RenderTargets[0] = FRenderTargetBinding(DestinationTexture, ERenderTargetLoadAction::ENoAction, /*InMipIndex = */0);
+	FIntVector DestinationSize = DestinationTexture->Desc.GetSize();
+
+	FPixelShaderUtils::AddFullscreenPass(
+		GraphBuilder,
+		ShaderMap,
+		RDG_EVENT_NAME("ConvertToNativeLandscapePatch"),
+		PixelShader,
+		ShaderParams,
+		FIntRect(0, 0, DestinationSize.X, DestinationSize.Y));
+}
+
+bool FConvertBackFromNativeLandscapePatchPS::ShouldCompilePermutation(const FGlobalShaderPermutationParameters& Parameters)
+{
+	return DoesPlatformSupportEditLayers(Parameters.Platform);
+}
+
+void FConvertBackFromNativeLandscapePatchPS::ModifyCompilationEnvironment(const FGlobalShaderPermutationParameters& Parameters, FShaderCompilerEnvironment& OutEnvironment)
+{
+	OutEnvironment.SetDefine(TEXT("CONVERT_BACK_FROM_NATIVE_LANDSCAPE_PATCH"), 1);
+}
+
+void FConvertBackFromNativeLandscapePatchPS::AddToRenderGraph(FRDGBuilder& GraphBuilder, FRDGTextureRef SourceTexture,
+	FRDGTextureRef DestinationTexture, const FLandscapeHeightPatchConvertToNativeParams& Params)
+{
+	FGlobalShaderMap* ShaderMap = GetGlobalShaderMap(GMaxRHIFeatureLevel);
+	TShaderMapRef<FConvertBackFromNativeLandscapePatchPS> PixelShader(ShaderMap);
+
+	FRDGTextureSRVRef SourceSRV = GraphBuilder.CreateSRV(FRDGTextureSRVDesc::CreateForMipLevel(SourceTexture, 0));
+
+	FConvertBackFromNativeLandscapePatchPS::FParameters* ShaderParams = GraphBuilder.AllocParameters<FConvertBackFromNativeLandscapePatchPS::FParameters>();
+	ShaderParams->InHeightmap = GraphBuilder.CreateSRV(FRDGTextureSRVDesc::CreateForMipLevel(SourceTexture, 0));
+	ShaderParams->InZeroInEncoding = Params.ZeroInEncoding;
+	ShaderParams->InHeightScale = Params.HeightScale;
+	ShaderParams->InHeightOffset = Params.HeightOffset;
+	ShaderParams->RenderTargets[0] = FRenderTargetBinding(DestinationTexture, ERenderTargetLoadAction::ENoAction, /*InMipIndex = */0);
+	FIntVector DestinationSize = DestinationTexture->Desc.GetSize();
+
+	FPixelShaderUtils::AddFullscreenPass(
+		GraphBuilder,
+		ShaderMap,
+		RDG_EVENT_NAME("ConvertBackFromNativeLandscapePatch"),
+		PixelShader,
+		ShaderParams,
+		FIntRect(0, 0, DestinationSize.X, DestinationSize.Y));
+}
+
+bool FApplyLandscapeTextureWeightPatchPS::ShouldCompilePermutation(const FGlobalShaderPermutationParameters& Parameters)
+{
+	return DoesPlatformSupportEditLayers(Parameters.Platform);
+}
+
+void FApplyLandscapeTextureWeightPatchPS::ModifyCompilationEnvironment(const FGlobalShaderPermutationParameters& Parameters, FShaderCompilerEnvironment& OutEnvironment)
+{
+	OutEnvironment.SetDefine(TEXT("APPLY_WEIGHT_PATCH"), 1);
+
+	// Make our flag choices match in the shader.
+	OutEnvironment.SetDefine(TEXT("RECTANGULAR_FALLOFF_FLAG"), static_cast<uint8>(FApplyLandscapeTextureHeightPatchPS::EFlags::RectangularFalloff));
+	OutEnvironment.SetDefine(TEXT("APPLY_PATCH_ALPHA_FLAG"), static_cast<uint8>(FApplyLandscapeTextureHeightPatchPS::EFlags::ApplyPatchAlpha));
+
+	OutEnvironment.SetDefine(TEXT("ADDITIVE_MODE"), static_cast<uint8>(FApplyLandscapeTextureHeightPatchPS::EBlendMode::Additive));
+	OutEnvironment.SetDefine(TEXT("ALPHA_BLEND_MODE"), static_cast<uint8>(FApplyLandscapeTextureHeightPatchPS::EBlendMode::AlphaBlend));
+	OutEnvironment.SetDefine(TEXT("MIN_MODE"), static_cast<uint8>(FApplyLandscapeTextureHeightPatchPS::EBlendMode::Min));
+	OutEnvironment.SetDefine(TEXT("MAX_MODE"), static_cast<uint8>(FApplyLandscapeTextureHeightPatchPS::EBlendMode::Max));
+}
+
+void FApplyLandscapeTextureWeightPatchPS::AddToRenderGraph(FRDGBuilder& GraphBuilder, FParameters* InParameters, const FIntRect& DestinationBounds)
+{
+	FGlobalShaderMap* ShaderMap = GetGlobalShaderMap(GMaxRHIFeatureLevel);
+	TShaderMapRef<FApplyLandscapeTextureWeightPatchPS> PixelShader(ShaderMap);
+
+	FPixelShaderUtils::AddFullscreenPass(
+		GraphBuilder,
+		ShaderMap,
+		RDG_EVENT_NAME("LandscapeTextureWeightPatch"),
+		PixelShader,
+		InParameters,
+		DestinationBounds);
+}
+
+bool FReinitializeLandscapePatchPS::ShouldCompilePermutation(const FGlobalShaderPermutationParameters& Parameters)
+{
+	return DoesPlatformSupportEditLayers(Parameters.Platform);
+}
+
+void FReinitializeLandscapePatchPS::ModifyCompilationEnvironment(const FGlobalShaderPermutationParameters& Parameters, FShaderCompilerEnvironment& OutEnvironment)
+{
+	OutEnvironment.SetDefine(TEXT("REINITIALIZE_PATCH"), 1);
+}
+
+void FReinitializeLandscapePatchPS::AddToRenderGraph(FRDGBuilder& GraphBuilder, FParameters* InParameters)
+{
+	FGlobalShaderMap* ShaderMap = GetGlobalShaderMap(GMaxRHIFeatureLevel);
+	TShaderMapRef<FReinitializeLandscapePatchPS> PixelShader(ShaderMap);
+
+	FIntVector DestinationSize = InParameters->RenderTargets[0].GetTexture()->Desc.GetSize();
+
+	FPixelShaderUtils::AddFullscreenPass(
+		GraphBuilder,
+		ShaderMap,
+		RDG_EVENT_NAME("LandscapeTextureHeightPatch"),
+		PixelShader,
+		InParameters,
+		FIntRect(0, 0, DestinationSize.X, DestinationSize.Y));
+}
+
+}//end UE::Landscape
+
+IMPLEMENT_GLOBAL_SHADER(UE::Landscape::FApplyLandscapeTextureHeightPatchPS, "/Plugin/LandscapePatch/Private/LandscapeTexturePatchPS.usf", "ApplyLandscapeTextureHeightPatch", SF_Pixel);
+IMPLEMENT_GLOBAL_SHADER(UE::Landscape::FOffsetHeightmapPS, "/Plugin/LandscapePatch/Private/LandscapeTexturePatchPS.usf", "ApplyOffsetToHeightmap", SF_Pixel);
+IMPLEMENT_GLOBAL_SHADER(UE::Landscape::FSimpleTextureCopyPS, "/Plugin/LandscapePatch/Private/LandscapeTexturePatchPS.usf", "SimpleTextureCopy", SF_Pixel);
+IMPLEMENT_GLOBAL_SHADER(UE::Landscape::FConvertToNativeLandscapePatchPS, "/Plugin/LandscapePatch/Private/LandscapeTexturePatchPS.usf", "ConvertToNativeLandscapePatch", SF_Pixel);
+IMPLEMENT_GLOBAL_SHADER(UE::Landscape::FConvertBackFromNativeLandscapePatchPS, "/Plugin/LandscapePatch/Private/LandscapeTexturePatchPS.usf", "ConvertBackFromNativeLandscapePatch", SF_Pixel);
+IMPLEMENT_GLOBAL_SHADER(UE::Landscape::FApplyLandscapeTextureWeightPatchPS, "/Plugin/LandscapePatch/Private/LandscapeTexturePatchPS.usf", "ApplyLandscapeTextureWeightPatch", SF_Pixel);
+IMPLEMENT_GLOBAL_SHADER(UE::Landscape::FReinitializeLandscapePatchPS, "/Plugin/LandscapePatch/Private/LandscapeTexturePatchPS.usf", "ReinitializePatch", SF_Pixel);
diff --git a/Engine/Plugins/Experimental/LandscapePatch/Source/LandscapePatch/Public/LandscapePatchComponent.h b/Engine/Plugins/Experimental/LandscapePatch/Source/LandscapePatch/Public/LandscapePatchComponent.h
--- a/Engine/Plugins/Experimental/LandscapePatch/Source/LandscapePatch/Public/LandscapePatchComponent.h
+++ b/Engine/Plugins/Experimental/LandscapePatch/Source/LandscapePatch/Public/LandscapePatchComponent.h
@@ -33,6 +33,8 @@ public:
 		UTextureRenderTarget2D* InCombinedResult,
 		const FName& InWeightmapLayerName) { return InCombinedResult; }
 
+	virtual bool IsAffectingWeightmapLayer(const FName& InLayerName) const { return false; }
+
 	UFUNCTION(BlueprintCallable, Category = "LandscapePatch")
 	void RequestLandscapeUpdate();
 
@@ -47,7 +49,7 @@ public:
 	 * @return false if the patch is marked as disabled and therefore can't affect the landscape.
 	 */
 	UFUNCTION(BlueprintCallable, Category = "LandscapePatch")
-	bool IsEnabled() { return bIsEnabled; }
+	bool IsEnabled() const { return bIsEnabled; }
 
 	UFUNCTION(BlueprintCallable, Category = "LandscapePatch")
 	FTransform GetLandscapeHeightmapCoordsToWorld() const;
diff --git a/Engine/Plugins/Experimental/LandscapePatch/Source/LandscapePatch/Public/LandscapePatchManager.h b/Engine/Plugins/Experimental/LandscapePatch/Source/LandscapePatch/Public/LandscapePatchManager.h
--- a/Engine/Plugins/Experimental/LandscapePatch/Source/LandscapePatch/Public/LandscapePatchManager.h
+++ b/Engine/Plugins/Experimental/LandscapePatch/Source/LandscapePatch/Public/LandscapePatchManager.h
@@ -51,8 +51,12 @@ public:
 
 	bool RemovePatch(TObjectPtr<ULandscapePatchComponent> Patch);
 
-	// UObject
+	
 #if WITH_EDITOR
+	// ALandscapeBlueprintBrushBase
+	virtual bool IsAffectingWeightmapLayer(const FName& InLayerName) const override;
+
+	// UObject
 	virtual void PostEditUndo() override;
 #endif
 	virtual bool IsEditorOnly() const override { return true; }
diff --git a/Engine/Plugins/Experimental/LandscapePatch/Source/LandscapePatch/Public/LandscapePatchUtil.h b/Engine/Plugins/Experimental/LandscapePatch/Source/LandscapePatch/Public/LandscapePatchUtil.h
new file mode 100644
--- /dev/null
+++ b/Engine/Plugins/Experimental/LandscapePatch/Source/LandscapePatch/Public/LandscapePatchUtil.h
@@ -0,0 +1,13 @@
+// Copyright Epic Games, Inc. All Rights Reserved.
+
+#pragma once
+
+#include "CoreMinimal.h"
+
+namespace UE {
+namespace Landscape {
+namespace PatchUtil {
+
+	void LANDSCAPEPATCH_API CopyTextureOnRenderThread(FRHICommandListImmediate& RHICmdList, const FTextureResource& Source, FTextureResource& Destination);
+
+}}}//end UE::Landscape::PatchUtil
\ No newline at end of file
diff --git a/Engine/Plugins/Experimental/LandscapePatch/Source/LandscapePatch/Public/LandscapeTextureBackedRenderTarget.h b/Engine/Plugins/Experimental/LandscapePatch/Source/LandscapePatch/Public/LandscapeTextureBackedRenderTarget.h
new file mode 100644
--- /dev/null
+++ b/Engine/Plugins/Experimental/LandscapePatch/Source/LandscapePatch/Public/LandscapeTextureBackedRenderTarget.h
@@ -0,0 +1,116 @@
+// Copyright Epic Games, Inc. All Rights Reserved.
+
+#pragma once
+
+#include "CoreMinimal.h"
+
+#include "Engine/TextureRenderTarget2D.h" // ETextureRenderTargetFormat
+#include "LandscapeTexturePatchPS.h" // FLandscapeHeightPatchConvertToNativeParams
+
+#include "LandscapeTextureBackedRenderTarget.generated.h"
+
+class UTexture2D;
+class UTextureRenderTarget2D;
+
+/**
+ * A combination of a render target and UTexture2D that allows render target to be saved across save/load/etc
+ * by copying the data back and forth from the internal texture.
+ * 
+ * After an Initialize() call, the internal texture will always be present, and the render target will be
+ * present depending on how SetUseInternalTextureOnly is called (by default, present). 
+ */
+UCLASS(Abstract)
+class ULandscapeTextureBackedRenderTargetBase : public UObject
+{
+	GENERATED_BODY()
+
+public:
+
+	UPROPERTY()
+	TObjectPtr<UTextureRenderTarget2D> PostLoadRT = nullptr;
+
+	virtual UTextureRenderTarget2D* GetRenderTarget() { return RenderTarget; }
+	virtual UTexture2D* GetInternalTexture() { return InternalTexture; }
+
+	/**
+	 * 
+	 */
+	virtual void SetUseInternalTextureOnly(bool bUseInternalTextureOnlyIn, bool bCopyExisting = true);
+
+	virtual void SetResolution(int32 SizeXIn, int32 SizeYIn);
+
+	virtual void Initialize();
+
+	virtual void CopyToInternalTexture() PURE_VIRTUAL(ULandscapeTextureBackedRenderTargetBase::CopyToInternalTexture);
+	virtual void CopyBackFromInternalTexture() PURE_VIRTUAL(ULandscapeTextureBackedRenderTargetBase::CopyBackFromInternalTexture);
+
+#if WITH_EDITOR
+	// UObject
+	virtual void PreSave(FObjectPreSaveContext SaveContext) override;
+	virtual void PostLoad() override;
+	virtual void PreDuplicate(FObjectDuplicationParameters& DupParams) override;
+	virtual void ExportCustomProperties(FOutputDevice& Out, uint32 Indent) override;
+	virtual void PostEditImport() override;
+#endif // WITH_EDITOR
+
+protected:
+	virtual ETextureSourceFormat GetInternalTextureFormat() PURE_VIRTUAL(ULandscapeTextureBackedRenderTargetBase::GetInternalTextureFormat, return ETextureSourceFormat::TSF_G8;);
+	virtual ETextureRenderTargetFormat GetRenderTargetFormat() PURE_VIRTUAL(ULandscapeTextureBackedRenderTargetBase::GetRenderTargetFormat, return ETextureRenderTargetFormat::RTF_R8;);
+
+	UPROPERTY(VisibleAnywhere, Category = InternalData)
+	TObjectPtr<UTexture2D> InternalTexture = nullptr;
+
+	UPROPERTY(VisibleAnywhere, Category = InternalData, Transient, DuplicateTransient)
+	TObjectPtr<UTextureRenderTarget2D> RenderTarget = nullptr;
+
+	UPROPERTY(VisibleAnywhere, Category = InternalData)
+	int32 SizeX = 32;
+
+	UPROPERTY(VisibleAnywhere, Category = InternalData)
+	int32 SizeY = 32;
+
+	UPROPERTY()
+	bool bUseInternalTextureOnly = false;
+};
+
+UCLASS()
+class ULandscapeWeightTextureBackedRenderTarget : public ULandscapeTextureBackedRenderTargetBase
+{
+	GENERATED_BODY()
+
+public:
+
+	virtual void SetUseAlphaChannel(bool bUseAlphaChannelIn);
+
+	virtual void CopyToInternalTexture() override;
+	virtual void CopyBackFromInternalTexture() override;
+protected:
+	virtual ETextureSourceFormat GetInternalTextureFormat() override;
+	virtual ETextureRenderTargetFormat GetRenderTargetFormat() override;
+
+	UPROPERTY()
+	bool bUseAlphaChannel = true;
+};
+
+UCLASS()
+class ULandscapeHeightTextureBackedRenderTarget : public ULandscapeTextureBackedRenderTargetBase
+{
+	GENERATED_BODY()
+
+public:
+
+	virtual void SetFormat(ETextureRenderTargetFormat FormatToUse);
+
+	virtual void CopyToInternalTexture() override;
+	virtual void CopyBackFromInternalTexture() override;
+
+	UPROPERTY()
+	FLandscapeHeightPatchConvertToNativeParams ConversionParams;
+
+protected:
+	virtual ETextureSourceFormat GetInternalTextureFormat() override { return ETextureSourceFormat::TSF_BGRA8; }
+	virtual ETextureRenderTargetFormat GetRenderTargetFormat() override { return RenderTargetFormat; }
+
+	UPROPERTY()
+	TEnumAsByte<ETextureRenderTargetFormat> RenderTargetFormat = ETextureRenderTargetFormat::RTF_R32f;
+};
diff --git a/Engine/Plugins/Experimental/LandscapePatch/Source/LandscapePatch/Public/LandscapeTextureHeightPatch.h b/Engine/Plugins/Experimental/LandscapePatch/Source/LandscapePatch/Public/LandscapeTextureHeightPatch.h
--- a/Engine/Plugins/Experimental/LandscapePatch/Source/LandscapePatch/Public/LandscapeTextureHeightPatch.h
+++ b/Engine/Plugins/Experimental/LandscapePatch/Source/LandscapePatch/Public/LandscapeTextureHeightPatch.h
@@ -5,8 +5,8 @@
 #include "CoreMinimal.h"
 
 #include "Engine/TextureRenderTarget2D.h"
-
-#include "LandscapeTextureHeightPatchPS.h" // FApplyLandscapeTextureHeightPatchPS::FParameters, FApplyLandscapeTextureHeightPatchPS::EBlendMode
+#include "LandscapeTexturePatchPS.h" // FApplyLandscapeTextureHeightPatchPS::FParameters, FApplyLandscapeTextureHeightPatchPS::EBlendMode
+#include "LandscapeTexturePatch.h" // ELandscapeTextureHeightPatchEncoding, ELandscapeTextureHeightPatchZeroHeightMeaning, FLandscapeTexturePatchEncodingSettings
 #include "LandscapeTexturePatchBase.h"
 #include "MatrixTypes.h"
 
@@ -63,66 +63,6 @@ enum class ELandscapeTextureHeightPatchFalloffMode : uint8
 	RoundedRectangle,
 };
 
-UENUM(BlueprintType)
-enum class ELandscapeTextureHeightPatchEncoding : uint8
-{
-	// Values in texture should be interpreted as being floats in the range [0,1]. User specifies what
-	// value corresponds to height 0 (i.e. height when landscape is "cleared"), and the size of the 
-	// range in world units.
-	ZeroToOne,
-
-	// Values in texture are direct world-space heights.
-	WorldUnits,
-
-	// Values in texture are stored the same way they are in landscape actors: as 16 bit integers packed 
-	// into two bytes, mapping to [-256, 256 - 1/128] before applying landscape scale.
-	NativePackedHeight
-
-	//~ Note that currently ZeroToOne and WorldUnits actually work the same way- we subtract the center point (0 for WorldUnits),
-	//~ then scale in some way (1.0 for WorldUnits). However, having separate options here allows us to initialize defaults
-	//~ appropriately when setting the encoding mode.
-};
-
-UENUM(BlueprintType)
-enum class ELandscapeTextureHeightPatchZeroHeightMeaning : uint8
-{
-	// Zero height corresponds to the patch vertical position relative to the landscape. This moves
-	// the results up and down as the patch moves up and down.
-	PatchZ,
-
-	// Zero height corresponds to Z = 0 in the local space of the landscape, regardless of the patch vertical
-	// position. For instance, if landscape transform has z=-100 in world, then writing height 0 will correspond
-	// to z=-100 in world coordinates, regardless of patch Z. 
-	LandscapeZ,
-
-	// Zero height corresponds to the height of the world origin relative to landscape. In other words, writing
-	// height 0 will correspond to world z = 0 regardless of patch Z or landscape transform (as long as landscape
-	// transform still has Z up in world coordinates).
-	WorldZero
-};
-
-//~ A struct in case we find that we need other encoding settings.
-USTRUCT(BlueprintType)
-struct LANDSCAPEPATCH_API FLandscapeTexturePatchEncodingSettings
-{
-	GENERATED_BODY()
-public:
-	/**
-	 * The value in the patch data that corresponds to 0 landscape height (which is in line with patch Z when
-	 * "Use Patch Z As Reference" is true, and at landscape zero/mid value when false).
-	 */
-	UPROPERTY(EditAnywhere, BlueprintReadWrite, Category = Settings)
-	double ZeroInEncoding = 0;
-
-	/**
-	 * The scale that should be aplied to the data stored in the patch relative to the zero in the encoding, in world coordinates.
-	 * For instance if the encoding is [0,1], and 0.5 correponds to 0, a WorldSpaceEncoding Scale of 100 means that the resulting
-	 * values will lie in the range [-50, 50] in world space, which would be [-0.5, 0.5] in the landscape local heights if the Z
-	 * scale is 100.
-	 */
-	UPROPERTY(EditAnywhere, BlueprintReadWrite, Category = Settings)
-	double WorldSpaceEncodingScale = 1;
-};
 
 /**
  * A texture-based height patch
@@ -352,7 +292,7 @@ private:
 	// These are used to convert render targets whose formats are not the usual two-channel 16bit values that
 	// we store in an internal texture. We convert these render targets to the native landscape height representation
 	// and store that in an internal texture.
-	UE::Landscape::FConvertToNativeLandscapePatchParams GetConversionParams();
+	FLandscapeHeightPatchConvertToNativeParams GetConversionParams();
 	void ConvertInternalRenderTargetToNativeTexture(bool bBlock);
 	void ConvertInternalRenderTargetBackFromNativeTexture(bool bLoading = false);
 
diff --git a/Engine/Plugins/Experimental/LandscapePatch/Source/LandscapePatch/Public/LandscapeTextureHeightPatchPS.h b/Engine/Plugins/Experimental/LandscapePatch/Source/LandscapePatch/Public/LandscapeTextureHeightPatchPS.h
deleted file mode 100644
diff --git a/Engine/Plugins/Experimental/LandscapePatch/Source/LandscapePatch/Public/LandscapeTexturePatch.h b/Engine/Plugins/Experimental/LandscapePatch/Source/LandscapePatch/Public/LandscapeTexturePatch.h
new file mode 100644
--- /dev/null
+++ b/Engine/Plugins/Experimental/LandscapePatch/Source/LandscapePatch/Public/LandscapeTexturePatch.h
@@ -0,0 +1,584 @@
+// Copyright Epic Games, Inc. All Rights Reserved.
+
+#pragma once
+
+#include "CoreMinimal.h"
+
+#include "LandscapePatchComponent.h"
+#include "LandscapeTextureBackedRenderTarget.h"
+
+#include "LandscapeTexturePatch.generated.h"
+
+class ULandscapeHeightTextureBackedRenderTarget;
+class UTexture;
+class UTexture2D;
+
+/**
+ * Determines where the patch gets its information, which affects its memory usage in editor (not in runtime,
+ * since patches are baked directly into landscape and removed for runtime).
+ */
+UENUM(BlueprintType)
+enum class ELandscapeTexturePatchSourceMode : uint8
+{
+	/**
+	 * The patch is considered not to have any data stored for this element. Setting source mode to this is
+	 * a way to discard any internally stored data.
+	 */
+	None,
+
+	/**
+	 * The data will be read from an internally-stored UTexture2D. In this mode, the patch can't be written-to via
+	 * blueprints, but it avoids storing the extra render target needed for TextureBackedRenderTarget.
+	 */
+	InternalTexture,
+
+	/**
+	* The patch data will be read from an internally-stored render target, which can be written to via Blueprints
+	* and which gets serialized to an internally stored UTexture2D when needed. Uses double the memory of InternalTexture.
+	*/
+	TextureBackedRenderTarget,
+
+	/**
+	 * The data will be read from a UTexture asset (which can be a render target). Allows multiple patches
+	 * to share the same texture.
+	 */
+	 TextureAsset
+};
+
+// Determines how the patch is combined with the previous state of the landscape.
+UENUM(BlueprintType)
+enum class ELandscapeTexturePatchBlendMode : uint8
+{
+	// Let the patch specify the actual target height, and blend that with the existing
+	// height using falloff/alpha. E.g. with no falloff and alpha 1, the landscape will
+	// be set directly to the height sampled from patch. With alpha 0.5, landscape height 
+	// will be averaged evenly with patch height.
+	AlphaBlend = static_cast<uint8>(UE::Landscape::FApplyLandscapeTextureHeightPatchPS::EBlendMode::AlphaBlend),
+
+	// Interpreting the landscape mid value as 0, use the texture patch as an offset to
+	// apply to the landscape. Falloff/alpha will just affect the degree to which the offset
+	// is applied (e.g. alpha of 0.5 will apply just half the offset).
+	Additive = static_cast<uint8>(UE::Landscape::FApplyLandscapeTextureHeightPatchPS::EBlendMode::Additive),
+
+	// Like Alpha Blend mode, but limited to only lowering the existing landscape values
+	Min = static_cast<uint8>(UE::Landscape::FApplyLandscapeTextureHeightPatchPS::EBlendMode::Min),
+
+	// Like Alpha Blend mode, but limited to only raising the existing landscape values
+	Max = static_cast<uint8>(UE::Landscape::FApplyLandscapeTextureHeightPatchPS::EBlendMode::Max)
+};
+
+// Determines falloff method for the patch's influence.
+UENUM(BlueprintType)
+enum class ELandscapeTexturePatchFalloffMode : uint8
+{
+	// Affect landscape in a circle inscribed in the patch, and fall off across
+	// a margin extending into that circle.
+	Circle,
+
+	// Affect entire rectangle of patch (except for circular corners), and fall off
+	// across a margin extending inward from the boundary.
+	RoundedRectangle,
+};
+
+UENUM(BlueprintType)
+enum class ELandscapeTextureHeightPatchEncoding : uint8
+{
+	// Values in texture should be interpreted as being floats in the range [0,1]. User specifies what
+	// value corresponds to height 0 (i.e. height when landscape is "cleared"), and the size of the 
+	// range in world units.
+	ZeroToOne,
+
+	// Values in texture are direct world-space heights.
+	WorldUnits,
+
+	// Values in texture are stored the same way they are in landscape actors: as 16 bit integers packed 
+	// into two bytes, mapping to [-256, 256 - 1/128] before applying landscape scale.
+	NativePackedHeight
+
+	//~ Note that currently ZeroToOne and WorldUnits actually work the same way- we subtract the center point (0 for WorldUnits),
+	//~ then scale in some way (1.0 for WorldUnits). However, having separate options here allows us to initialize defaults
+	//~ appropriately when setting the encoding mode via ResetSourceEncodingMode.
+};
+
+UENUM(BlueprintType)
+enum class ELandscapeTextureHeightPatchZeroHeightMeaning : uint8
+{
+	// Zero height corresponds to the patch vertical position relative to the landscape. This moves
+	// the results up and down as the patch moves up and down.
+	PatchZ,
+
+	// Zero height corresponds to Z = 0 in the local space of the landscape, regardless of the patch vertical
+	// position. For instance, if landscape transform has z=-100 in world, then writing height 0 will correspond
+	// to z=-100 in world coordinates, regardless of patch Z. 
+	LandscapeZ,
+
+	// Zero height corresponds to the height of the world origin relative to landscape. In other words, writing
+	// height 0 will correspond to world z = 0 regardless of patch Z or landscape transform (as long as landscape
+	// transform still has Z up in world coordinates).
+	WorldZero
+};
+
+//~ A struct in case we find that we need other encoding settings.
+USTRUCT(BlueprintType)
+struct LANDSCAPEPATCH_API FLandscapeTexturePatchEncodingSettings
+{
+	GENERATED_BODY()
+public:
+	/**
+	 * The value in the patch data that corresponds to 0 landscape height (which is in line with patch Z when
+	 * "Use Patch Z As Reference" is true, and at landscape zero/mid value when false).
+	 */
+	UPROPERTY(EditAnywhere, BlueprintReadWrite, Category = Settings)
+	double ZeroInEncoding = 0;
+
+	/**
+	 * The scale that should be aplied to the data stored in the patch relative to the zero in the encoding, in world coordinates.
+	 * For instance if the encoding is [0,1], and 0.5 correponds to 0, a WorldSpaceEncoding Scale of 100 means that the resulting
+	 * values will lie in the range [-50, 50] in world space, which would be [-0.5, 0.5] in the landscape local heights if the Z
+	 * scale is 100.
+	 */
+	UPROPERTY(EditAnywhere, BlueprintReadWrite, Category = Settings)
+	double WorldSpaceEncodingScale = 1;
+};
+
+/**
+ * Helper class for ULandscapeTexturePatch that stores information for a given weight layer.
+ * Should not be used outside this class, but does need to be a UObject (so can't be nested).
+ */
+UCLASS()
+class LANDSCAPEPATCH_API ULandscapeWeightPatchTextureInfo : public UObject
+{
+	GENERATED_BODY()
+
+public:
+#if WITH_EDITOR
+	// UObject
+	virtual void PostEditChangeProperty(FPropertyChangedEvent& PropertyChangedEvent) override;
+#endif // WITH_EDITOR
+
+protected:
+	UPROPERTY(EditAnywhere, Category = WeightPatch)
+	FName WeightmapLayerName;
+	
+	/** Texture to use when source mode is set to texture asset. */
+	UPROPERTY(EditAnywhere, Category = WeightPatch, meta = (EditConditionHides,
+		EditCondition = "SourceMode == ELandscapeTexturePatchSourceMode::TextureAsset", 
+		DisallowedAssetDataTags = "VirtualTextureStreaming=True"))
+	TObjectPtr<UTexture> TextureAsset = nullptr;
+
+	/** Not directly settable via detail panel- for display/debugging purposes only. */
+	UPROPERTY(VisibleAnywhere, Category = WeightPatch, Instanced, AdvancedDisplay)
+	TObjectPtr<ULandscapeWeightTextureBackedRenderTarget> InternalData = nullptr;
+
+	UPROPERTY(EditAnywhere, Category = WeightPatch, meta = (EditConditionHides, EditCondition = "false"))
+	ELandscapeTexturePatchSourceMode SourceMode = ELandscapeTexturePatchSourceMode::None;
+
+	UPROPERTY(EditAnywhere, Category = WeightPatch, meta = (DisplayName = "Source Mode"))
+	ELandscapeTexturePatchSourceMode DetailPanelSourceMode = ELandscapeTexturePatchSourceMode::None;
+
+	//~ We could refactor things such that we always have an InternalData pointer, even when we use
+	//~ a texture asset, and then we could use the boolean inside that instead (which needs to be there
+	//~ so that we know how many channels we need). Not clear whether that will be any cleaner though.
+	UPROPERTY(EditAnywhere, Category = WeightPatch)
+	bool bUseAlphaChannel = false;
+
+	// Can't make TOptional a UPROPERTY, hence these two.
+	UPROPERTY()
+	bool bOverrideBlendMode = false;
+	UPROPERTY()
+	ELandscapeTexturePatchBlendMode OverrideBlendMode = ELandscapeTexturePatchBlendMode::AlphaBlend;
+
+	// TODO: We could support having different per-layer falloff modes and falloff amounts as well, as
+	// additional override members. But probably better to wait to see if that is actually desired.
+
+	// Needed mainly so that we can get at the resolution...
+	UPROPERTY()
+	ULandscapeTexturePatch* OwningPatch = nullptr;
+
+	// TODO: Like the similar flag for the height patch, this might not work once local merge works
+	// with landscape brushes...
+	bool bReinitializeOnNextRender = false;
+
+	void SetSourceMode(ELandscapeTexturePatchSourceMode NewMode);
+
+	friend class ULandscapeTexturePatch;
+};
+
+UCLASS(Blueprintable, BlueprintType, ClassGroup = Landscape, meta = (BlueprintSpawnableComponent))
+class LANDSCAPEPATCH_API ULandscapeTexturePatch : public ULandscapePatchComponent
+{
+	GENERATED_BODY()
+
+public:
+
+#if WITH_EDITOR
+	virtual UTextureRenderTarget2D* Render_Native(bool InIsHeightmap,
+		UTextureRenderTarget2D* InCombinedResult,
+		const FName& InWeightmapLayerName) override;
+
+	// ULandscapePatchComponent
+	virtual bool IsAffectingWeightmapLayer(const FName& InLayerName) const override;
+
+	// UObject
+	virtual void PostEditChangeProperty(FPropertyChangedEvent& PropertyChangedEvent) override;
+#endif
+
+	/**
+	 * Gets the transform from patch to world. The transform is based off of the component
+	 * transform, but with rotation changed to align to the landscape, only using the yaw
+	 * to rotate it relative to the landscape.
+	 */
+	UFUNCTION(BlueprintCallable, Category = LandscapePatch)
+	virtual FTransform GetPatchToWorldTransform() const;
+
+	/**
+	 * Gives size in unscaled world coordinates (ie before applying patch transform) of the patch as measured 
+	 * between the centers of the outermost pixels. Measuring the coverage this way means that a patch can 
+	 * affect the same region of the landscape regardless of patch resolution.
+	 * This is also the range across which bilinear interpolation always has correct values, so the area outside 
+	 * this center portion is usually set as a "dead" border that doesn't affect the landscape.
+	 */
+	UFUNCTION(BlueprintCallable, Category = LandscapePatch)
+	virtual FVector2D GetUnscaledCoverage() const { return FVector2D(UnscaledPatchCoverage); }
+
+	/**
+	 * Set the patch coverage (see GetUnscaledCoverage for description).
+	 */
+	UFUNCTION(BlueprintCallable, Category = LandscapePatch)
+	virtual void SetUnscaledCoverage(FVector2D Coverage) { UnscaledPatchCoverage = Coverage; }
+
+	/**
+	 * Gives size in unscaled world coordinates of the patch in the world, based off of UnscaledCoverage and
+	 * texture resolution (i.e., adds a half-pixel around UnscaledCoverage).
+	 */
+	UFUNCTION(BlueprintCallable, Category = LandscapePatch)
+	virtual FVector2D GetFullUnscaledWorldSize() const;
+
+	/**
+	 * Gets the size (in pixels) of the internal textures used by the patch. Does not reflect the resolution
+	 * of any used texture assets (if the source mode is texture asset).
+	 */
+	UFUNCTION(BlueprintCallable, Category = LandscapePatch)
+	virtual FVector2D GetResolution() const { return FVector2D(ResolutionX, ResolutionY); }
+
+	/**
+	 * Sets the resolution of the currently used internal texture or render target. Has no effect
+	 * if the source mode is set to an external texture asset.
+	 *
+	 * @return true if successful.
+	 */
+	UFUNCTION(BlueprintCallable, Category = LandscapePatch)
+	virtual void SetResolution(FVector2D ResolutionIn);
+
+	/**
+	 * Given the landscape resolution, current patch coverage, and a landscape resolution multiplier, gives the
+	 * needed resolution of the landscape patch. I.e., figures out the number of pixels in the landcape that
+	 * would be in a region of such size, and then uses the resolution multiplier to give a result.
+	 *
+	 * @return true if successful (may fail if landscape is not set, for instance)
+	 */
+	UFUNCTION(BlueprintCallable, Category = LandscapePatch, meta = (ResolutionMultiplier = "1.0"))
+	virtual UPARAM(DisplayName = "Success") bool GetInitResolutionFromLandscape(float ResolutionMultiplier, FVector2D& ResolutionOut) const;
+
+
+	UFUNCTION(BlueprintCallable, Category = LandscapePatch)
+	void SetFalloff(float FalloffIn) 
+	{
+		Modify();
+		Falloff = FalloffIn; 
+	}
+
+	UFUNCTION(BlueprintCallable, Category = LandscapePatch)
+	void SetBlendMode(ELandscapeTexturePatchBlendMode BlendModeIn) 
+	{ 
+		Modify();
+		BlendMode = BlendModeIn; 
+	}
+
+
+	// Height related functions:
+
+	UFUNCTION(BlueprintCallable, Category = LandscapePatch)
+	virtual ELandscapeTexturePatchSourceMode GetHeightSourceMode() const { return HeightSourceMode; }
+
+	/**
+	 * Changes source mode. When changing between sources, existing data is copied from one to the other
+	 * when possible.
+	 */
+	UFUNCTION(BlueprintCallable, Category = LandscapePatch)
+	virtual void SetHeightSourceMode(ELandscapeTexturePatchSourceMode NewMode);
+
+	/**
+	 * Sets the texture used for height when the height source mode is set to texture asset. Note that
+	 * virtual textures are not supported.
+	 */
+	UFUNCTION(BlueprintCallable, Category = LandscapePatch)
+	void SetHeightTextureAsset(UTexture* TextureIn);
+
+	UFUNCTION(BlueprintCallable, Category = LandscapePatch)
+	virtual UTextureRenderTarget2D* GetHeightRenderTarget() 
+	{ 
+		return HeightInternalData ? HeightInternalData->GetRenderTarget() : nullptr; 
+	}
+
+	UFUNCTION(BlueprintCallable, Category = LandscapePatch, meta = (ETextureRenderTargetFormat = "ETextureRenderTargetFormat::RTF_R32f"))
+	void SetHeightRenderTargetFormat(ETextureRenderTargetFormat Format);
+
+	UFUNCTION(BlueprintCallable, Category = "LandscapePatch")
+	void SetUseAlphaChannelForHeight(bool bUse)
+	{ 
+		Modify();
+		bUseTextureAlphaForHeight = bUse; 
+	}
+
+	/**
+	 * Set the height encoding mode for the patch, which determines how stored values in the patch
+	 * are translated into heights when applying to landscape.
+	 */
+	UFUNCTION(BlueprintCallable, Category = LandscapePatch)
+	void SetHeightEncodingMode(ELandscapeTextureHeightPatchEncoding EncodingMode)
+	{
+		Modify();
+		HeightEncoding = EncodingMode;
+	}
+
+	/**
+	 * Just like SetSourceEncodingMode, but resets ZeroInEncoding and WorldSpaceEncodingScale to mode-specific defaults.
+	 */
+	UFUNCTION(BlueprintCallable, Category = LandscapePatch)
+	void ResetHeightEncodingMode(ELandscapeTextureHeightPatchEncoding EncodingMode);
+
+	/**
+	 * Set settings that determine how values in the patch are translated into heights. This is only
+	 * used if the encoding mode is not NativePackedHeight, where values are expected to be already
+	 * in the same space as the landscape heightmap.
+	 */
+	UFUNCTION(BlueprintCallable, Category = LandscapePatch)
+	void SetHeightEncodingSettings(const FLandscapeTexturePatchEncodingSettings& Settings)
+	{
+		Modify();
+		HeightEncodingSettings = Settings;
+	}
+
+	/**
+	 * Set how zero height is interpreted, see comments in ELandscapeTextureHeightPatchZeroHeightMeaning.
+	 */
+	UFUNCTION(BlueprintCallable, Category = "LandscapePatch")
+	void SetZeroHeightMeaning(ELandscapeTextureHeightPatchZeroHeightMeaning ZeroHeightMeaningIn)
+	{ 
+		Modify();
+		ZeroHeightMeaning = ZeroHeightMeaningIn;
+	}
+
+
+	// Weight related functions:
+
+	/**
+	 * By default, the layer is added with source mode set to be a texture-backed render target.
+	 */
+	UFUNCTION(BlueprintCallable, Category = LandscapePatch)
+	virtual void AddWeightPatch(const FName& InWeightmapLayerName, ELandscapeTexturePatchSourceMode SourceMode, bool bUseAlphaChannel);
+
+	UFUNCTION(BlueprintCallable, Category = LandscapePatch)
+	virtual void RemoveWeightPatch(const FName& InWeightmapLayerName);
+
+	UFUNCTION(BlueprintCallable, Category = LandscapePatch)
+	virtual void RemoveAllWeightPatches();
+
+	/** Sets the soure mode of all weight patches to "None". */
+	UFUNCTION(BlueprintCallable, Category = LandscapePatch)
+	virtual void DisableAllWeightPatches();
+
+	UFUNCTION(BlueprintCallable, Category = LandscapePatch)
+	TArray<FName> GetAllWeightPatchLayerNames();
+	
+	UFUNCTION(BlueprintCallable, Category = LandscapePatch)
+	virtual ELandscapeTexturePatchSourceMode GetWeightPatchSourceMode(const FName& InWeightmapLayerName);
+
+	UFUNCTION(BlueprintCallable, Category = LandscapePatch)
+	virtual void SetWeightPatchSourceMode(const FName& InWeightmapLayerName, ELandscapeTexturePatchSourceMode NewMode);
+
+	UFUNCTION(BlueprintCallable, Category = LandscapePatch)
+	void SetWeightPatchTextureAsset(const FName& InWeightmapLayerName, UTexture* TextureIn);
+
+	UFUNCTION(BlueprintCallable, Category = LandscapePatch)
+	virtual UTextureRenderTarget2D* GetWeightPatchRenderTarget(const FName& InWeightmapLayerName);
+
+	UFUNCTION(BlueprintCallable, Category = LandscapePatch)
+	virtual void SetUseAlphaChannelForWeightPatch(const FName& InWeightmapLayerName, bool bUseAlphaChannel);
+
+	UFUNCTION(BlueprintCallable, Category = LandscapePatch)
+	virtual void SetWeightPatchBlendModeOverride(const FName& InWeightmapLayerName, ELandscapeTexturePatchBlendMode BlendMode);
+
+	UFUNCTION(BlueprintCallable, Category = LandscapePatch)
+	virtual void ClearWeightPatchBlendModeOverride(const FName& InWeightmapLayerName);
+
+
+	// These need to be public so that we can take the internal textures and write them to external ones,
+	// but unclear whether we want to expose them to blueprints, since this is a fairly internal thing.
+	UTexture2D* GetHeightInternalTexture() { return HeightInternalData ? HeightInternalData->GetInternalTexture() : nullptr; ; }
+	UTexture2D* GetWeightPatchInternalTexture(const FName& InWeightmapLayerName);
+
+protected:
+
+	UPROPERTY()
+	int32 ResolutionX = 32;
+	UPROPERTY()
+	int32 ResolutionY = 32;
+
+	/** At scale 1.0, the X and Y of the region affected by the height patch. This corresponds to the distance from the center
+	 of the first pixel to the center of the last pixel in the patch texture in the X and Y directions. */
+	UPROPERTY(EditAnywhere, Category = Settings, meta = (UIMin = "0", ClampMin = "0"))
+	FVector2D UnscaledPatchCoverage = FVector2D(2000, 2000);
+
+	UPROPERTY(EditAnywhere, Category = Settings)
+	ELandscapeTexturePatchBlendMode BlendMode = ELandscapeTexturePatchBlendMode::AlphaBlend;
+
+	UPROPERTY(EditAnywhere, Category = Settings)
+	ELandscapeTexturePatchFalloffMode FalloffMode = ELandscapeTexturePatchFalloffMode::RoundedRectangle;
+
+	/**
+	 * Distance (in unscaled world coordinates) across which to smoothly fall off the patch effects.
+	 */
+	UPROPERTY(EditAnywhere, Category = Settings, meta = (ClampMin = "0", UIMax = "2000"))
+	float Falloff = 0;
+
+
+	// Height properties:
+
+	// How the heightmap of the patch is stored. This is the property that is actually used, and it will
+	// agree with DetailPanelHeightSourceMode at all times except when user is changing the latter via the
+	// detail panel.
+	//~ TODO: The property specifiers here are a hack to force this (hidden) property to be preserved across reruns of
+	//~ a construction script in a blueprint actor. We should find the proper way that this is supposed to be done.
+	UPROPERTY(EditAnywhere, Category = HeightPatch, meta = (EditConditionHides, EditCondition = "false"))
+	ELandscapeTexturePatchSourceMode HeightSourceMode = ELandscapeTexturePatchSourceMode::None;
+
+	/**
+	 * How the heightmap of the patch is stored.
+	 */
+	UPROPERTY(EditAnywhere, Category = HeightPatch, meta = (DisplayName = "Source Mode"))
+	ELandscapeTexturePatchSourceMode DetailPanelHeightSourceMode = ELandscapeTexturePatchSourceMode::None;
+
+	/** Not directly settable via detail panel- for display/debugging purposes only. */
+	UPROPERTY(VisibleAnywhere, Category = HeightPatch, Instanced)
+	TObjectPtr<ULandscapeHeightTextureBackedRenderTarget> HeightInternalData = nullptr;
+
+	/**
+	 * Texture used when source mode is set to a texture asset.
+	 */
+	UPROPERTY(EditAnywhere, Category = HeightPatch, meta = (EditConditionHides,
+		EditCondition = "HeightSourceMode == ELandscapeTexturePatchSourceMode::TextureAsset", 
+		DisallowedAssetDataTags = "VirtualTextureStreaming=True"))
+	TObjectPtr<UTexture> HeightTextureAsset = nullptr;
+
+	
+	/** When true, texture alpha channel will be used when applying the patch. */
+	UPROPERTY(EditAnywhere, Category = HeightPatch)
+	bool bUseTextureAlphaForHeight = false;
+
+	/** How the values stored in the patch represent the height. Not customizable for Internal Texture source mode, which always uses native packed height. */
+	UPROPERTY(EditAnywhere, Category = HeightPatch, meta = (
+		EditCondition = "HeightSourceMode != ELandscapeTexturePatchSourceMode::InternalTexture"))
+	ELandscapeTextureHeightPatchEncoding HeightEncoding = ELandscapeTextureHeightPatchEncoding::ZeroToOne;
+
+	/** Encoding settings. Not relevant when using native packed height as the encoding. */
+	UPROPERTY(EditAnywhere, Category = HeightPatch, meta = (UIMin = "0", UIMax = "1",
+		EditCondition = "HeightSourceMode != ELandscapeTexturePatchSourceMode::InternalTexture && HeightEncoding != ELandscapeTextureHeightPatchEncoding::NativePackedHeight"))
+	FLandscapeTexturePatchEncodingSettings HeightEncodingSettings;
+
+	/**
+	 * How 0 height is interpreted.
+	 */
+	UPROPERTY(EditAnywhere, Category = HeightPatch)
+	ELandscapeTextureHeightPatchZeroHeightMeaning ZeroHeightMeaning = ELandscapeTextureHeightPatchZeroHeightMeaning::PatchZ;
+
+	/**
+	 * Whether to apply the patch Z scale to the height stored in the patch.
+	 */
+	UPROPERTY(EditAnywhere, Category = HeightPatch, AdvancedDisplay, meta = (DisplayName = "Apply Component Z Scale"))
+	bool bApplyComponentZScale = true;
+
+
+	// Weight properties:
+
+	/** */
+	UPROPERTY(EditAnywhere, Category = WeightPatches, Instanced)
+	TArray<TObjectPtr<ULandscapeWeightPatchTextureInfo>> WeightPatches;
+
+	// Used to detect changes to the number of weight patches via the detail panel, so that we can
+	// initialize the "owner" pointer when patches are added or trigger update when patches are removed.
+	UPROPERTY()
+	int32 NumWeightPatches = 0;
+
+
+	// Reinitialization from detail panel:
+
+	/**
+	 * Given the current initialization settings, reinitialize the height patch.
+	 */
+	UFUNCTION(CallInEditor, Category = HeightPatch)
+	void ReinitializeHeight();
+
+	UFUNCTION(CallInEditor, Category = WeightPatches)
+	void ReinitializeWeights();
+
+	bool bReinitializeHeightOnNextRender = false;
+
+	/**
+	 * Adjusts patch rotation to be aligned to a 90 degree increment relative to the landscape,
+	 * adjusts UnscaledPatchCoverage such that it becomes a multiple of landscape quad size, and
+	 * adjusts patch location so that the boundaries of the covered area lie on the nearest
+	 * landscape vertices.
+	 * Note that this doesn't adjust the resolution of the texture that the patch uses, so landscape
+	 * vertices within the inside of the patch may still not always align with texture patch pixel
+	 * centers (if the resolutions aren't multiples of each other).
+	 */
+	UFUNCTION(BlueprintCallable, CallInEditor, Category = Initialization)
+	void SnapToLandscape();
+
+	/** When initializing from landscape, set resolution based off of the landscape (and a multiplier). */
+	UPROPERTY(EditAnywhere, Category = Initialization)
+	bool bBaseResolutionOffLandscape = true;
+
+	/** 
+	 * Multiplier to apply to landscape resolution when initializing patch resolution. A value greater than 1.0 will use higher
+	 * resolution than the landscape (perhaps useful for slightly more accurate results while not aligned to landscape), and
+	 * a value less that 1.0 will use lower.
+	 */
+	UPROPERTY(EditAnywhere, Category = Initialization, meta = (
+		EditCondition = "bBaseResolutionOffLandscape"))
+	float ResolutionMultiplier = 1;
+
+	/** Texture width to use when reinitializing, if not basing resolution off landscape. */
+	UPROPERTY(EditAnywhere, Category = Initialization, meta = (EditCondition = "!bBaseResolutionOffLandscape", 
+		ClampMin = "1"))
+	int32 InitTextureSizeX = 33;
+
+	/** Texture height to use when reinitializing */
+	UPROPERTY(EditAnywhere, Category = Initialization, meta = (EditCondition = "!bBaseResolutionOffLandscape", 
+		ClampMin = "1"))
+	int32 InitTextureSizeY = 33;
+
+private:
+#if WITH_EDITOR
+	FLandscapeHeightPatchConvertToNativeParams GetHeightConversionParams() const;
+	UTextureRenderTarget2D* ApplyToHeightmap(UTextureRenderTarget2D* InCombinedResult);
+	UTextureRenderTarget2D* ApplyToWeightmap(ULandscapeWeightPatchTextureInfo* PatchInfo, UTextureRenderTarget2D* InCombinedResult);
+
+	void GetCommonShaderParams(const FIntPoint& SourceResolutionIn, const FIntPoint& DestinationResolutionIn, 
+		FTransform& PatchToWorldOut, FVector2f& PatchWorldDimensionsOut, FMatrix44f& HeightmapToPatchOut, 
+		FIntRect& DestinationBoundsOut, FVector2f& EdgeUVDeadBorderOut, float& FalloffWorldMarginOut) const;
+	void GetHeightShaderParams(const FIntPoint& SourceResolutionIn, const FIntPoint& DestinationResolutionIn,
+		UE::Landscape::FApplyLandscapeTextureHeightPatchPS::FParameters& ParamsOut, FIntRect& DestinationBoundsOut) const;
+	void GetWeightShaderParams(const FIntPoint& SourceResolutionIn, const FIntPoint& DestinationResolutionIn,
+		const ULandscapeWeightPatchTextureInfo* WeightPatchInfo, 
+		UE::Landscape::FApplyLandscapeTextureWeightPatchPS::FParameters& ParamsOut, FIntRect& DestinationBoundsOut) const;
+	FMatrix44f GetPatchToHeightmapUVs(int32 PatchSizeX, int32 PatchSizeY, int32 HeightmapSizeX, int32 HeightmapSizeY) const;
+	void ReinitializeHeight(UTextureRenderTarget2D* InCombinedResult);
+	void ReinitializeWeightPatch(ULandscapeWeightPatchTextureInfo* PatchInfo, UTextureRenderTarget2D* InCombinedResult);
+#endif // WITH_EDITOR
+
+	UPROPERTY()
+	TEnumAsByte<ETextureRenderTargetFormat> HeightRenderTargetFormat = ETextureRenderTargetFormat::RTF_R32f;
+};
\ No newline at end of file
diff --git a/Engine/Plugins/Experimental/LandscapePatch/Source/LandscapePatch/Public/LandscapeTexturePatchBase.h b/Engine/Plugins/Experimental/LandscapePatch/Source/LandscapePatch/Public/LandscapeTexturePatchBase.h
--- a/Engine/Plugins/Experimental/LandscapePatch/Source/LandscapePatch/Public/LandscapeTexturePatchBase.h
+++ b/Engine/Plugins/Experimental/LandscapePatch/Source/LandscapePatch/Public/LandscapeTexturePatchBase.h
@@ -5,6 +5,7 @@
 #include "CoreMinimal.h"
 
 #include "LandscapePatchComponent.h"
+#include "LandscapeTexturePatch.h" // ELandscapeTexturePatchSourceMode
 
 #include "LandscapeTexturePatchBase.generated.h"
 
@@ -12,32 +13,6 @@ class UTexture;
 class UTexture2D;
 class UTextureRenderTarget2D;
 
-/**
- * Determines how the patch stores its information, which affects its memory usage in editor (not in runtime,
- * since patches are baked directly into landscape and removed for runtime).
- */
-UENUM(BlueprintType)
-enum class ELandscapeTexturePatchSourceMode : uint8
-{
-	/**
-	 * The data will be read from an internally-stored UTexture2D. In this mode, the patch can't be written-to via 
-	 * blueprints, but it avoids storing the extra render target needed for TextureBackedRenderTarget.
-	 */
-	 InternalTexture,
-
-	 /**
-	 * The patch data will be read from an internally-stored render target, which can be written to via Blueprints
-	 * and which gets serialized to an internally stored UTexture2D when needed. Uses double the memory of InternalTexture.
-	 */
-	TextureBackedRenderTarget,
-
-	 /**
-	  * The data will be read from a UTexture asset (which can be a render target). Allows multiple patches
-	  * to share the same texture.
-	  */
-	  TextureAsset
-};
-
 /**
  * A texture-based landscape patch. Base class for height patches (and, to be implemented, for weight patches).
  */
diff --git a/Engine/Plugins/Experimental/LandscapePatch/Source/LandscapePatch/Public/LandscapeTexturePatchPS.h b/Engine/Plugins/Experimental/LandscapePatch/Source/LandscapePatch/Public/LandscapeTexturePatchPS.h
new file mode 100644
--- /dev/null
+++ b/Engine/Plugins/Experimental/LandscapePatch/Source/LandscapePatch/Public/LandscapeTexturePatchPS.h
@@ -0,0 +1,291 @@
+// Copyright Epic Games, Inc. All Rights Reserved.
+
+#pragma once
+
+#include "CoreMinimal.h"
+
+#include "GlobalShader.h"
+#include "ShaderParameterStruct.h"
+
+#include "LandscapeTexturePatchPS.generated.h"
+
+// Values needed to convert a patch stored in some source encoding into the native (two byte int) encoding and back
+USTRUCT()
+struct FLandscapeHeightPatchConvertToNativeParams
+{
+	GENERATED_BODY()
+
+	UPROPERTY()
+	float ZeroInEncoding;
+
+	UPROPERTY()
+	float HeightScale;
+
+	UPROPERTY()
+	float HeightOffset;
+};
+
+namespace UE::Landscape
+{
+
+/**
+ * Shader that applies a texture-based height patch to a landscape heightmap.
+ */
+class LANDSCAPEPATCH_API FApplyLandscapeTextureHeightPatchPS : public FGlobalShader
+{
+	DECLARE_GLOBAL_SHADER(FApplyLandscapeTextureHeightPatchPS);
+	SHADER_USE_PARAMETER_STRUCT(FApplyLandscapeTextureHeightPatchPS, FGlobalShader);
+
+public:
+
+	enum class EBlendMode : uint8
+	{
+		/** Desired height is alpha blended with the current. */
+		AlphaBlend,
+
+		/** Desired height is multiplied by alpha and added to current. */
+		Additive,
+
+		/** Like AlphaBlend, but patch is limited to only lowering the landscape. */
+		Min,
+
+		/** Like AlphaBlend, but patch is limited to only raising the landscape. */
+		Max
+	};
+
+	// Flags that get packed into a bitfield because we're not allowed to use bool shader parameters:
+	enum class EFlags : uint8
+	{
+		None = 0,
+
+		// When false, falloff is circular.
+		RectangularFalloff = 1 << 0,
+
+		// When true, the texture alpha channel is considered for blending (in addition to falloff, if nonzero)
+		ApplyPatchAlpha = 1 << 1,
+
+		// When false, the input is directly interpreted as being the height value to process. When true, the height
+		// is unpacked from the red and green channels to make a 16 bit int.
+		InputIsPackedHeight = 1 << 2
+	};
+
+	// TODO: We could consider exposing an additional global alpha setting that we can use to pass in the given
+	// edit layer alpha value... On the other hand, we currently don't bother doing this in any existing blueprint
+	// brushes, and it would be hard to support in a way that doesn't require each blueprint brush to respect it
+	// individually... Not clear whether this is something worth doing yet.
+
+	BEGIN_SHADER_PARAMETER_STRUCT(FParameters, )
+		SHADER_PARAMETER_RDG_TEXTURE_SRV(Texture2D<float4>, InSourceHeightmap)
+		SHADER_PARAMETER_RDG_TEXTURE_SRV(Texture2D<float4>, InHeightPatch)
+		SHADER_PARAMETER_SAMPLER(SamplerState, InHeightPatchSampler)
+		SHADER_PARAMETER(FMatrix44f, InHeightmapToPatch)
+		// Value in patch that corresponds to the landscape mid value, which is our "0 height".
+		SHADER_PARAMETER(float, InZeroInEncoding)
+		// Scale to apply to source values relative to the value that represents 0 height.
+		SHADER_PARAMETER(float, InHeightScale)
+		// Offset to apply to height result after applying height scale
+		SHADER_PARAMETER(float, InHeightOffset)
+		// Amount of the patch edge to not apply in UV space. Generally set to 0.5/Dimensions to avoid applying
+		// the edge half-pixels.
+		SHADER_PARAMETER(FVector2f, InEdgeUVDeadBorder)
+		// In world units, the size of the margin across which the alpha falls from 1 to 0
+		SHADER_PARAMETER(float, InFalloffWorldMargin)
+		// Size of the patch in world units (used for falloff)
+		SHADER_PARAMETER(FVector2f, InPatchWorldDimensions)
+		SHADER_PARAMETER(uint32, InBlendMode)
+		// Some combination of the flags (see constants above).
+		SHADER_PARAMETER(uint32, InFlags)
+
+		RENDER_TARGET_BINDING_SLOTS() // Holds our output
+		END_SHADER_PARAMETER_STRUCT()
+
+	static bool ShouldCompilePermutation(const FGlobalShaderPermutationParameters& Parameters);
+	static void ModifyCompilationEnvironment(const FGlobalShaderPermutationParameters& InParameters, FShaderCompilerEnvironment& OutEnvironment);
+	static void AddToRenderGraph(FRDGBuilder& GraphBuilder, FParameters* InParameters, const FIntRect& DestinationBounds);
+};
+
+ENUM_CLASS_FLAGS(FApplyLandscapeTextureHeightPatchPS::EFlags);
+
+/**
+ * Simple shader that just offsets each height value in a height patch by a constant.
+ */
+class LANDSCAPEPATCH_API FOffsetHeightmapPS : public FGlobalShader
+{
+	DECLARE_GLOBAL_SHADER(FOffsetHeightmapPS);
+	SHADER_USE_PARAMETER_STRUCT(FOffsetHeightmapPS, FGlobalShader);
+
+public:
+
+	BEGIN_SHADER_PARAMETER_STRUCT(FParameters, )
+		SHADER_PARAMETER_RDG_TEXTURE_SRV(Texture2D<float4>, InHeightmap)
+		SHADER_PARAMETER(float, InHeightOffset)
+		RENDER_TARGET_BINDING_SLOTS() // Holds our output
+		END_SHADER_PARAMETER_STRUCT()
+
+	static bool ShouldCompilePermutation(const FGlobalShaderPermutationParameters& Parameters);
+	static void ModifyCompilationEnvironment(const FGlobalShaderPermutationParameters& InParameters, FShaderCompilerEnvironment& OutEnvironment);
+	static void AddToRenderGraph(FRDGBuilder& GraphBuilder, FParameters* InParameters);
+};
+
+/**
+ * Shader that converts a texture stored in some other encoding (where height is in the R channel) to the 
+ * landscape "native" encoding, where height is stored as a 16 bit int split across the R and G channels.
+ * This is not perfectly reversible (in case of clamping and due to rounding), but it lets us store the
+ * texture in the way that it would be applied to the landscape (usually).
+ */
+class LANDSCAPEPATCH_API FConvertToNativeLandscapePatchPS : public FGlobalShader
+{
+	DECLARE_GLOBAL_SHADER(FConvertToNativeLandscapePatchPS);
+	SHADER_USE_PARAMETER_STRUCT(FConvertToNativeLandscapePatchPS, FGlobalShader);
+
+public:
+
+	BEGIN_SHADER_PARAMETER_STRUCT(FParameters, )
+		SHADER_PARAMETER_RDG_TEXTURE_SRV(Texture2D<float4>, InHeightmap)
+		SHADER_PARAMETER(float, InZeroInEncoding)
+		SHADER_PARAMETER(float, InHeightScale)
+		SHADER_PARAMETER(float, InHeightOffset)
+		RENDER_TARGET_BINDING_SLOTS() // Holds our output
+		END_SHADER_PARAMETER_STRUCT()
+
+	static bool ShouldCompilePermutation(const FGlobalShaderPermutationParameters& Parameters);
+	static void ModifyCompilationEnvironment(const FGlobalShaderPermutationParameters& InParameters, FShaderCompilerEnvironment& OutEnvironment);
+	static void AddToRenderGraph(FRDGBuilder& GraphBuilder, FRDGTextureRef SourceTexture,
+		FRDGTextureRef DestinationTexture, const FLandscapeHeightPatchConvertToNativeParams& Params);
+};
+
+/**
+ * Shader that undoes the conversion done by FConvertToNativeLandscapePatchPS (to the extent possible, since
+ * rounding/clamping makes it not perfectly recoverable).
+ */
+class LANDSCAPEPATCH_API FConvertBackFromNativeLandscapePatchPS : public FGlobalShader
+{
+	DECLARE_GLOBAL_SHADER(FConvertBackFromNativeLandscapePatchPS);
+	SHADER_USE_PARAMETER_STRUCT(FConvertBackFromNativeLandscapePatchPS, FGlobalShader);
+
+public:
+
+	BEGIN_SHADER_PARAMETER_STRUCT(FParameters, )
+		SHADER_PARAMETER_RDG_TEXTURE_SRV(Texture2D<float4>, InHeightmap)
+		SHADER_PARAMETER(float, InZeroInEncoding)
+		SHADER_PARAMETER(float, InHeightScale)
+		SHADER_PARAMETER(float, InHeightOffset)
+		RENDER_TARGET_BINDING_SLOTS() // Holds our output
+		END_SHADER_PARAMETER_STRUCT()
+
+	static bool ShouldCompilePermutation(const FGlobalShaderPermutationParameters& Parameters);
+	static void ModifyCompilationEnvironment(const FGlobalShaderPermutationParameters& InParameters, FShaderCompilerEnvironment& OutEnvironment);
+	static void AddToRenderGraph(FRDGBuilder& GraphBuilder, FRDGTextureRef SourceTexture,
+		FRDGTextureRef DestinationTexture, const FLandscapeHeightPatchConvertToNativeParams& Params);
+};
+
+class LANDSCAPEPATCH_API FApplyLandscapeTextureWeightPatchPS : public FGlobalShader
+{
+	DECLARE_GLOBAL_SHADER(FApplyLandscapeTextureWeightPatchPS);
+	SHADER_USE_PARAMETER_STRUCT(FApplyLandscapeTextureWeightPatchPS, FGlobalShader);
+
+public:
+
+	enum class EBlendMode : uint8
+	{
+		/** Desired height is alpha blended with the current. */
+		AlphaBlend,
+
+		/** Desired height is multiplied by alpha and added to current (and clamped). */
+		Additive,
+
+		/** Like AlphaBlend, but patch is limited to only lowering the landscape. */
+		Min,
+
+		/** Like AlphaBlend, but patch is limited to only raising the landscape. */
+		Max
+	};
+
+	// Flags that get packed into a bitfield because we're not allowed to use bool shader parameters:
+	enum class EFlags : uint8
+	{
+		None = 0,
+
+		// When false, falloff is circular.
+		RectangularFalloff = 1 << 0,
+
+		// When true, the texture alpha channel is considered for blending (in addition to falloff, if nonzero)
+		ApplyPatchAlpha = 1 << 1,
+	};
+
+	// TODO: We could consider exposing an additional global alpha setting that we can use to pass in the given
+	// edit layer alpha value... On the other hand, we currently don't bother doing this in any existing blueprint
+	// brushes, and it would be hard to support in a way that doesn't require each blueprint brush to respect it
+	// individually... Not clear whether this is something worth doing yet.
+
+	BEGIN_SHADER_PARAMETER_STRUCT(FParameters, )
+		SHADER_PARAMETER_RDG_TEXTURE_SRV(Texture2D<float4>, InSourceWeightmap)
+		SHADER_PARAMETER_RDG_TEXTURE_SRV(Texture2D<float4>, InWeightPatch)
+		SHADER_PARAMETER_SAMPLER(SamplerState, InWeightPatchSampler)
+		SHADER_PARAMETER(FMatrix44f, InWeightmapToPatch)
+		// Amount of the patch edge to not apply in UV space. Generally set to 0.5/Dimensions to avoid applying
+		// the edge half-pixels.
+		SHADER_PARAMETER(FVector2f, InEdgeUVDeadBorder)
+		// In world units, the size of the margin across which the alpha falls from 1 to 0
+		SHADER_PARAMETER(float, InFalloffWorldMargin)
+		// Size of the patch in world units (used for falloff)
+		SHADER_PARAMETER(FVector2f, InPatchWorldDimensions)
+		SHADER_PARAMETER(uint32, InBlendMode)
+		// Some combination of the flags (see constants above).
+		SHADER_PARAMETER(uint32, InFlags)
+
+		RENDER_TARGET_BINDING_SLOTS() // Holds our output
+		END_SHADER_PARAMETER_STRUCT()
+
+		static bool ShouldCompilePermutation(const FGlobalShaderPermutationParameters& Parameters);
+	static void ModifyCompilationEnvironment(const FGlobalShaderPermutationParameters& InParameters, FShaderCompilerEnvironment& OutEnvironment);
+	static void AddToRenderGraph(FRDGBuilder& GraphBuilder, FParameters* InParameters, const FIntRect& DestinationBounds);
+};
+
+ENUM_CLASS_FLAGS(FApplyLandscapeTextureWeightPatchPS::EFlags);
+
+class LANDSCAPEPATCH_API FReinitializeLandscapePatchPS : public FGlobalShader
+{
+	DECLARE_GLOBAL_SHADER(FReinitializeLandscapePatchPS);
+	SHADER_USE_PARAMETER_STRUCT(FReinitializeLandscapePatchPS, FGlobalShader);
+
+public:
+
+	BEGIN_SHADER_PARAMETER_STRUCT(FParameters, )
+		SHADER_PARAMETER_RDG_TEXTURE_SRV(Texture2D<float4>, InSource)
+		SHADER_PARAMETER_SAMPLER(SamplerState, InSourceSampler)
+		SHADER_PARAMETER(FMatrix44f, InPatchToSource)
+
+		RENDER_TARGET_BINDING_SLOTS() // Holds our output
+		END_SHADER_PARAMETER_STRUCT()
+
+	static bool ShouldCompilePermutation(const FGlobalShaderPermutationParameters& Parameters);
+	static void ModifyCompilationEnvironment(const FGlobalShaderPermutationParameters& InParameters, FShaderCompilerEnvironment& OutEnvironment);
+	static void AddToRenderGraph(FRDGBuilder& GraphBuilder, FParameters* InParameters);
+};
+
+/**
+ * Simple shader for copying textures of potentially different resolutions.
+ */
+//~ Theoretically CopyToResolveTarget or AddCopyToResolveTargetPass should work, but I was unable to make them
+//~ work without lots of complaints from the renderer.
+class LANDSCAPEPATCH_API FSimpleTextureCopyPS : public FGlobalShader
+{
+	DECLARE_GLOBAL_SHADER(FSimpleTextureCopyPS);
+	SHADER_USE_PARAMETER_STRUCT(FSimpleTextureCopyPS, FGlobalShader);
+
+public:
+
+	BEGIN_SHADER_PARAMETER_STRUCT(FParameters, )
+		SHADER_PARAMETER_RDG_TEXTURE_SRV(Texture2D<float4>, InSource)
+		SHADER_PARAMETER_SAMPLER(SamplerState, InSourceSampler)
+		SHADER_PARAMETER(FVector2f, InDestinationResolution)
+		RENDER_TARGET_BINDING_SLOTS() // Holds our output
+		END_SHADER_PARAMETER_STRUCT()
+
+	static bool ShouldCompilePermutation(const FGlobalShaderPermutationParameters& Parameters);
+	static void ModifyCompilationEnvironment(const FGlobalShaderPermutationParameters& InParameters, FShaderCompilerEnvironment& OutEnvironment);
+	static void AddToRenderGraph(FRDGBuilder& GraphBuilder, FRDGTextureRef SourceTexture, FRDGTextureRef DestinationTexture);
+};
+
+}//end UE::Landscape
\ No newline at end of file
diff --git a/Engine/Plugins/Experimental/LandscapePatch/Source/LandscapePatchEditorOnly/Private/LandscapePatchEditorOnlyModule.cpp b/Engine/Plugins/Experimental/LandscapePatch/Source/LandscapePatchEditorOnly/Private/LandscapePatchEditorOnlyModule.cpp
--- a/Engine/Plugins/Experimental/LandscapePatch/Source/LandscapePatchEditorOnly/Private/LandscapePatchEditorOnlyModule.cpp
+++ b/Engine/Plugins/Experimental/LandscapePatch/Source/LandscapePatchEditorOnly/Private/LandscapePatchEditorOnlyModule.cpp
@@ -3,6 +3,7 @@
 #include "LandscapePatchEditorOnlyModule.h"
 
 #include "Editor/UnrealEdEngine.h"
+#include "LandscapeTexturePatch.h"
 #include "LandscapeTexturePatchBase.h"
 #include "LandscapeTexturePatchCustomization.h"
 #include "LandscapeTexturePatchVisualizer.h"
@@ -17,6 +18,7 @@ void FLandscapePatchEditorOnlyModule::StartupModule()
 	{
 		TSharedPtr<FLandscapeTexturePatchVisualizer> Visualizer = MakeShared<FLandscapeTexturePatchVisualizer>();
 		GUnrealEd->RegisterComponentVisualizer(ULandscapeTexturePatchBase::StaticClass()->GetFName(), Visualizer);
+		GUnrealEd->RegisterComponentVisualizer(ULandscapeTexturePatch::StaticClass()->GetFName(), Visualizer);
 		// This call should maybe be inside the RegisterComponentVisualizer call above, but since it's not,
 		// we'll put it here.
 		Visualizer->OnRegister();
diff --git a/Engine/Plugins/Experimental/LandscapePatch/Source/LandscapePatchEditorOnly/Private/LandscapeTexturePatchVisualizer.cpp b/Engine/Plugins/Experimental/LandscapePatch/Source/LandscapePatchEditorOnly/Private/LandscapeTexturePatchVisualizer.cpp
--- a/Engine/Plugins/Experimental/LandscapePatch/Source/LandscapePatchEditorOnly/Private/LandscapeTexturePatchVisualizer.cpp
+++ b/Engine/Plugins/Experimental/LandscapePatch/Source/LandscapePatchEditorOnly/Private/LandscapeTexturePatchVisualizer.cpp
@@ -4,25 +4,34 @@
 
 #include "Landscape.h"
 #include "LandscapeTexturePatchBase.h"
+#include "LandscapeTexturePatch.h"
 
 #include "SceneManagement.h" // FPrimitiveDrawInterface
 
 void FLandscapeTexturePatchVisualizer::DrawVisualization(const UActorComponent* Component, const FSceneView* View, FPrimitiveDrawInterface* PDI)
 {
-	const ULandscapeTexturePatchBase* Patch = Cast<ULandscapeTexturePatchBase>(Component);
-	if (!ensure(Patch))
-	{
-		return;
-	}
-	
-	FTransform PatchToWorld = Patch->GetPatchToWorldTransform();
-	
 	FColor Color = FColor::Red;
 	float Thickness = 3;
 	float DepthBias = 1;
 	bool bScreenSpace = 1;
 
-	DrawRectangle(PDI, PatchToWorld.GetTranslation(), PatchToWorld.GetUnitAxis(EAxis::X), PatchToWorld.GetUnitAxis(EAxis::Y),
-		Color, Patch->GetUnscaledCoverage().X * PatchToWorld.GetScale3D().X, Patch->GetUnscaledCoverage().Y * PatchToWorld.GetScale3D().Y, SDPG_Foreground,
-		Thickness, DepthBias, bScreenSpace);
+	if (const ULandscapeTexturePatch* Patch = Cast<ULandscapeTexturePatch>(Component))
+	{
+		FTransform PatchToWorld = Patch->GetPatchToWorldTransform();
+		DrawRectangle(PDI, PatchToWorld.GetTranslation(), PatchToWorld.GetUnitAxis(EAxis::X), PatchToWorld.GetUnitAxis(EAxis::Y),
+			Color, Patch->GetUnscaledCoverage().X * PatchToWorld.GetScale3D().X, Patch->GetUnscaledCoverage().Y * PatchToWorld.GetScale3D().Y, SDPG_Foreground,
+			Thickness, DepthBias, bScreenSpace);
+	}
+	// TODO: Should be removed once we delete ULandscapeTexturePatchBase and ULandscapeTextureHeightPatch
+	else if (const ULandscapeTexturePatchBase* HeightPatch = Cast<ULandscapeTexturePatchBase>(Component))
+	{
+		FTransform PatchToWorld = HeightPatch->GetPatchToWorldTransform();
+		DrawRectangle(PDI, PatchToWorld.GetTranslation(), PatchToWorld.GetUnitAxis(EAxis::X), PatchToWorld.GetUnitAxis(EAxis::Y),
+			Color, HeightPatch->GetUnscaledCoverage().X * PatchToWorld.GetScale3D().X, HeightPatch->GetUnscaledCoverage().Y * PatchToWorld.GetScale3D().Y, SDPG_Foreground,
+			Thickness, DepthBias, bScreenSpace);
+	}
+	else
+	{
+		ensure(false);
+	}
 }
\ No newline at end of file
diff --git a/Engine/Source/Runtime/Landscape/Private/LandscapeEditLayers.cpp b/Engine/Source/Runtime/Landscape/Private/LandscapeEditLayers.cpp
--- a/Engine/Source/Runtime/Landscape/Private/LandscapeEditLayers.cpp
+++ b/Engine/Source/Runtime/Landscape/Private/LandscapeEditLayers.cpp
@@ -6785,11 +6785,15 @@ int32 ALandscape::PerformLayersWeightmapsGlobalMerge(FUpdateLayersContentContext
 
 							PrintLayersDebugRT(FString::Printf(TEXT("LS Weight: %s %s -> Brush %s"), *Layer.Name.ToString(), *LandscapeBrush->GetName(), *BrushOutputRT->GetName()), BrushOutputRT);
 
-							SourceDebugName = FString::Printf(TEXT("Weight: %s PaintLayer: %s Brush: %s"), *Layer.Name.ToString(), *LayerInfoObj->LayerName.ToString(), *BrushOutputRT->GetName());
-							DestDebugName = LandscapeScratchRT3->GetName();
+							// Copy result back if brush did not edit things in place.
+							if (BrushOutputRT != LandscapeScratchRT3)
+							{
+								SourceDebugName = FString::Printf(TEXT("Weight: %s PaintLayer: %s Brush: %s"), *Layer.Name.ToString(), *LayerInfoObj->LayerName.ToString(), *BrushOutputRT->GetName());
+								DestDebugName = LandscapeScratchRT3->GetName();
 
-							ExecuteCopyLayersTexture({ FLandscapeLayersCopyTextureParams(SourceDebugName, BrushOutputRT->GameThread_GetRenderTargetResource(), DestDebugName, LandscapeScratchRT3->GameThread_GetRenderTargetResource()) });
-							PrintLayersDebugRT(FString::Printf(TEXT("LS Weight: %s Component %s += -> Combined %s"), *Layer.Name.ToString(), *BrushOutputRT->GetName(), *LandscapeScratchRT3->GetName()), LandscapeScratchRT3);
+								ExecuteCopyLayersTexture({ FLandscapeLayersCopyTextureParams(SourceDebugName, BrushOutputRT->GameThread_GetRenderTargetResource(), DestDebugName, LandscapeScratchRT3->GameThread_GetRenderTargetResource()) });
+								PrintLayersDebugRT(FString::Printf(TEXT("LS Weight: %s Component %s += -> Combined %s"), *Layer.Name.ToString(), *BrushOutputRT->GetName(), *LandscapeScratchRT3->GetName()), LandscapeScratchRT3);
+							}
 						}
 
 						PrintLayersDebugRT(FString::Printf(TEXT("LS Weight: %s CombinedPostBrushProcLayerWeightmap -> Paint Layer RT %s"), *Layer.Name.ToString(), *LandscapeScratchRT3->GetName()), LandscapeScratchRT3, 0, false);
diff --git a/Engine/Source/Runtime/Landscape/Public/LandscapeBlueprintBrushBase.h b/Engine/Source/Runtime/Landscape/Public/LandscapeBlueprintBrushBase.h
--- a/Engine/Source/Runtime/Landscape/Public/LandscapeBlueprintBrushBase.h
+++ b/Engine/Source/Runtime/Landscape/Public/LandscapeBlueprintBrushBase.h
@@ -64,7 +64,7 @@ public:
 
 	bool IsAffectingHeightmap() const { return AffectHeightmap; }
 	bool IsAffectingWeightmap() const { return AffectWeightmap; }
-	bool IsAffectingWeightmapLayer(const FName& InLayerName) const;
+	virtual bool IsAffectingWeightmapLayer(const FName& InLayerName) const;
 	bool IsVisible() const { return bIsVisible; }
 	bool IsLayerUpdatePending() const;
 
