#LandscapePatch Forward declare fix

diff --git a/Engine/Plugins/Experimental/LandscapePatch/Source/LandscapePatch/Public/LandscapePatchUtil.h b/Engine/Plugins/Experimental/LandscapePatch/Source/LandscapePatch/Public/LandscapePatchUtil.h
--- a/Engine/Plugins/Experimental/LandscapePatch/Source/LandscapePatch/Public/LandscapePatchUtil.h
+++ b/Engine/Plugins/Experimental/LandscapePatch/Source/LandscapePatch/Public/LandscapePatchUtil.h
@@ -4,6 +4,9 @@
 
 #include "CoreMinimal.h"
 
+class FRHICommandListImmediate;
+class FTextureResource;
+
 namespace UE {
 namespace Landscape {
 namespace PatchUtil {
