LandscapePatch: Fix MSVC static analyer warning

diff --git a/Engine/Plugins/Experimental/LandscapePatch/Source/LandscapePatch/Private/LandscapeTexturePatch.cpp b/Engine/Plugins/Experimental/LandscapePatch/Source/LandscapePatch/Private/LandscapeTexturePatch.cpp
--- a/Engine/Plugins/Experimental/LandscapePatch/Source/LandscapePatch/Private/LandscapeTexturePatch.cpp
+++ b/Engine/Plugins/Experimental/LandscapePatch/Source/LandscapePatch/Private/LandscapeTexturePatch.cpp
@@ -529,7 +529,7 @@ void ULandscapeTexturePatch::ReinitializeWeights()
 	}
 	SetResolution(DesiredResolution);
 
-	const FName VisibilityLayerName = Landscape->VisibilityLayer ? Landscape->VisibilityLayer->LayerName : NAME_None;
+	const FName VisibilityLayerName = Landscape->VisibilityLayer != nullptr ? Landscape->VisibilityLayer->LayerName : NAME_None;
 
 	ULandscapeInfo* Info = Landscape->GetLandscapeInfo();
 	if (Info)
