#LandscapePatch Disable ensures for GetResource(). Resource is nullptr when using nullrhi (commandlets, cooking, etc)

diff --git a/Engine/Plugins/Experimental/LandscapePatch/Source/LandscapePatch/Private/LandscapeTextureBackedRenderTarget.cpp b/Engine/Plugins/Experimental/LandscapePatch/Source/LandscapePatch/Private/LandscapeTextureBackedRenderTarget.cpp
--- a/Engine/Plugins/Experimental/LandscapePatch/Source/LandscapePatch/Private/LandscapeTextureBackedRenderTarget.cpp
+++ b/Engine/Plugins/Experimental/LandscapePatch/Source/LandscapePatch/Private/LandscapeTextureBackedRenderTarget.cpp
@@ -243,7 +243,7 @@ void ULandscapeWeightTextureBackedRenderTarget::CopyToInternalTexture()
 		return;
 	}
 
-	if (!ensure(IsValid(InternalTexture) && InternalTexture->GetResource()))
+	if (!(ensure(IsValid(InternalTexture)) && InternalTexture->GetResource()))
 	{
 		Modify();
 		InternalTexture = CreateTexture(this);
@@ -271,7 +271,7 @@ void ULandscapeWeightTextureBackedRenderTarget::CopyBackFromInternalTexture()
 	InternalTexture->UpdateResource();
 	FTextureCompilingManager::Get().FinishCompilation({ InternalTexture });
 
-	if (!ensure(InternalTexture->GetResource()))
+	if (!InternalTexture->GetResource())
 	{
 		return;
 	}
@@ -341,7 +341,7 @@ void ULandscapeHeightTextureBackedRenderTarget::CopyToInternalTexture()
 		return;
 	}
 
-	if (!ensure(IsValid(InternalTexture) && InternalTexture->GetResource()))
+	if (!(ensure(IsValid(InternalTexture)) && InternalTexture->GetResource()))
 	{
 		Modify();
 		InternalTexture = CreateTexture(this);
@@ -405,7 +405,7 @@ void ULandscapeHeightTextureBackedRenderTarget::CopyBackFromInternalTexture()
 	InternalTexture->UpdateResource();
 	FTextureCompilingManager::Get().FinishCompilation({ InternalTexture });
 
-	if (!ensure(InternalTexture->GetResource()))
+	if (!InternalTexture->GetResource())
 	{
 		return;
 	}
