LandscapePatch: Prevent certain properties from being manipulated in the detail panel of the blueprint editor due to issues with instanced properties there (they can still be set via function calls in construction script). Prevent ensure from firing in headless mode. Prevent weird interactions with patch manager when setting patch manager in blueprint editor.

diff --git a/Engine/Plugins/Experimental/LandscapePatch/Source/LandscapePatch/Private/LandscapePatchComponent.cpp b/Engine/Plugins/Experimental/LandscapePatch/Source/LandscapePatch/Private/LandscapePatchComponent.cpp
--- a/Engine/Plugins/Experimental/LandscapePatch/Source/LandscapePatch/Private/LandscapePatchComponent.cpp
+++ b/Engine/Plugins/Experimental/LandscapePatch/Source/LandscapePatch/Private/LandscapePatchComponent.cpp
@@ -84,18 +84,14 @@ void ULandscapePatchComponent::OnComponentCreated()
 	using namespace LandscapePatchComponentLocals;
 
 	Super::OnComponentCreated();
-		
-	// Mark whether we're creating from scratch of from a copy
+
+	// Mark whether we're creating from scratch or from a copy
 	bWasCopy = bPropertiesCopiedIndicator;
 	bPropertiesCopiedIndicator = true;
 
 	UWorld* World = GetWorld();
-	AActor* OwningActor = GetOwner();
-
-	if (HasAnyFlags(RF_ClassDefaultObject | RF_ArchetypeObject))
+	if (IsTemplate() || World->WorldType != EWorldType::Editor)
 	{
-		// If this is a template object from the blueprint editor or something of the sort, don't do
-		// anything else special- we don't want to add to patch manager, etc.
 		return;
 	}
 
@@ -104,62 +100,43 @@ void ULandscapePatchComponent::OnComponentCreated()
 		// If we copied over a patch manager, presumably Landscape should be
 		// copied over as well, but might as well do this to be safe.
 		Landscape = PatchManager->GetOwningLandscape();
+
+		if (!PatchManager->ContainsPatch(this))
+		{
+			SetPatchManager(PatchManager.Get());
+		}
 	}
-	else if (World)
+	else if (Landscape.IsValid())
 	{
-		if (Landscape.IsValid())
+		// If we copied over a patch with a landscape but no manager.
+		SetLandscape(Landscape.Get());
+	}
+	else
+	{
+		// See if the level has a height patch manager to which we can add ourselves
+		TActorIterator<ALandscapePatchManager> ManagerIterator(World);
+		if (!!ManagerIterator)
 		{
-			// If we copied over a patch with a landscape but no manager, find or create a manager in that landscape
-			if (Landscape->CanHaveLayersContent())
-			{
-				ALandscapePatchManager* ManagerToUse = FindExistingPatchManagerForLandscape(Landscape.Get());
-				if (!ManagerToUse)
-				{
-					ManagerToUse = CreateNewPatchManagerForLandscape(Landscape.Get());
-				}
-
-				SetPatchManager(ManagerToUse);
-			}
-			else
-			{
-				UE_LOG(LogLandscapePatch, Warning, TEXT("Landscape target for height patch did not have edit layers enabled. Unable to create patch manager."));
-				Landscape = nullptr;
-			}
+			SetPatchManager(*ManagerIterator);
 		}
 		else
 		{
-			// Didn't copy over an existing manager or landscape.
-
-			// See if the level has a height patch manager to which we can add ourselves
-			TActorIterator<ALandscapePatchManager> ManagerIterator(World);
-			if (!!ManagerIterator)
-			{
-				SetPatchManager(*ManagerIterator);
-			}
-			else
+			// If no existing manager, find some landscape and add a new one.
+			for (TActorIterator<ALandscape> LandscapeIterator(World); LandscapeIterator; ++LandscapeIterator)
 			{
-				// If no existing manager, find some landscape and add a new one.
-				for (TActorIterator<ALandscape> LandscapeIterator(World); LandscapeIterator; ++LandscapeIterator)
+				if (LandscapeIterator->CanHaveLayersContent())
 				{
-					if (LandscapeIterator->CanHaveLayersContent())
-					{
-						Landscape = *LandscapeIterator;
-						SetPatchManager(CreateNewPatchManagerForLandscape(Landscape.Get()));
-						break;
-					}
-				}
-				if (!PatchManager.IsValid())
-				{
-					UE_LOG(LogLandscapePatch, Warning, TEXT("Unable to find a landscape with edit layers enabled. Unable to create patch manager."));
+					Landscape = *LandscapeIterator;
+					SetPatchManager(CreateNewPatchManagerForLandscape(Landscape.Get()));
+					break;
 				}
 			}
+			if (!PatchManager.IsValid())
+			{
+				UE_LOG(LogLandscapePatch, Warning, TEXT("Unable to find a landscape with edit layers enabled. Unable to create patch manager."));
+			}
 		}
 	}
-
-	if (PatchManager.IsValid())
-	{
-		PatchManager->AddPatch(this);
-	}
 }
 
 void ULandscapePatchComponent::PostLoad()
@@ -222,6 +199,13 @@ void ULandscapePatchComponent::PostEditChangeProperty(FPropertyChangedEvent& Pro
 
 	Super::PostEditChangeProperty(PropertyChangedEvent);
 
+	// Do a bunch of checks to make sure that we don't try to do anything when the editing is happening inside the blueprint editor.
+	UWorld* World = GetWorld();
+	if (IsTemplate() || !IsValid(this) || !IsValid(World) || World->WorldType != EWorldType::Editor)
+	{
+		return;
+	}
+
 	// If we're changing the owning landscape or patch manaager, there's some work we need to do to remove/add 
 	// ourselves from/to the proper brush managers.
 	if (PropertyChangedEvent.Property 
@@ -249,8 +233,8 @@ void ULandscapePatchComponent::SetLandscape(ALandscape* NewLandscape)
 	// to be able to use this function to reorder patches. On the other hand, it seems inconvenient
 	// to accidentally swap patch managers if there are multiple in the same landscape, and we kept
 	// landscape the same. It's hard to know the ideal behavior, but for now we'll keep it.
-	if (Landscape.Get() == NewLandscape && ((Landscape.IsNull() && !PatchManager)
-		|| (Landscape.IsValid() && PatchManager && PatchManager->GetOwningLandscape() == Landscape.Get())))
+	if (Landscape.Get() == NewLandscape && (
+		(Landscape.IsNull() && PatchManager.IsNull()) || (Landscape.IsValid() && PatchManager.IsValid() && PatchManager->GetOwningLandscape() == Landscape.Get())))
 	{
 		return;
 	}
@@ -263,6 +247,12 @@ void ULandscapePatchComponent::SetLandscape(ALandscape* NewLandscape)
 		return;
 	}
 
+	UWorld* World = GetWorld();
+	if (IsTemplate() || !IsValid(World) || World->WorldType != EWorldType::Editor)
+	{
+		return;
+	}
+
 	// If landscape was valid, try to find a patch manager inside that landscape.
 	if (Landscape->CanHaveLayersContent())
 	{
@@ -303,7 +293,8 @@ void ULandscapePatchComponent::SetPatchManager(ALandscapePatchManager* NewPatchM
 	PatchManager = NewPatchManager;
 	if (NewPatchManager)
 	{
-		if (!HasAnyFlags(RF_ClassDefaultObject | RF_ArchetypeObject))
+		UWorld* World = GetWorld();
+		if (!IsTemplate() && IsValid(World) && World->WorldType == EWorldType::Editor)
 		{
 			PatchManager->AddPatch(this);
 		}
diff --git a/Engine/Plugins/Experimental/LandscapePatch/Source/LandscapePatch/Private/LandscapePatchManager.cpp b/Engine/Plugins/Experimental/LandscapePatch/Source/LandscapePatch/Private/LandscapePatchManager.cpp
--- a/Engine/Plugins/Experimental/LandscapePatch/Source/LandscapePatch/Private/LandscapePatchManager.cpp
+++ b/Engine/Plugins/Experimental/LandscapePatch/Source/LandscapePatch/Private/LandscapePatchManager.cpp
@@ -171,6 +171,11 @@ void ALandscapePatchManager::SetTargetLandscape(ALandscape* InTargetLandscape)
 #endif
 }
 
+bool ALandscapePatchManager::ContainsPatch(TObjectPtr<ULandscapePatchComponent> Patch)
+{
+	return PatchComponents.Contains(Patch);
+}
+
 void ALandscapePatchManager::AddPatch(TObjectPtr<ULandscapePatchComponent> Patch)
 {
 	if (Patch)
@@ -189,7 +194,10 @@ bool ALandscapePatchManager::RemovePatch(TObjectPtr<ULandscapePatchComponent> Pa
 	{
 		Modify();
 		bRemoved = PatchComponents.Remove(TSoftObjectPtr<ULandscapePatchComponent>(Patch.Get())) > 0;
-		RequestLandscapeUpdate();
+		if (bRemoved)
+		{
+			RequestLandscapeUpdate();
+		}
 	}
 	
 	return bRemoved;
diff --git a/Engine/Plugins/Experimental/LandscapePatch/Source/LandscapePatch/Private/LandscapeTextureBackedRenderTarget.cpp b/Engine/Plugins/Experimental/LandscapePatch/Source/LandscapePatch/Private/LandscapeTextureBackedRenderTarget.cpp
--- a/Engine/Plugins/Experimental/LandscapePatch/Source/LandscapePatch/Private/LandscapeTextureBackedRenderTarget.cpp
+++ b/Engine/Plugins/Experimental/LandscapePatch/Source/LandscapePatch/Private/LandscapeTextureBackedRenderTarget.cpp
@@ -238,6 +238,11 @@ void ULandscapeWeightTextureBackedRenderTarget::CopyToInternalTexture()
 #if WITH_EDITOR
 	using namespace LandscapeTextureBackedRenderTargetLocals;
 
+	if (!IsCopyingBackAndForthAllowed())
+	{
+		return;
+	}
+
 	if (!ensure(IsValid(RenderTarget)))
 	{
 		return;
@@ -264,6 +269,11 @@ void ULandscapeWeightTextureBackedRenderTarget::CopyBackFromInternalTexture()
 #if WITH_EDITOR
 	using namespace LandscapeTextureBackedRenderTargetLocals;
 
+	if (!IsCopyingBackAndForthAllowed())
+	{
+		return;
+	}
+
 	if (!ensure(IsValid(InternalTexture)))
 	{
 		return;
@@ -336,6 +346,11 @@ void ULandscapeHeightTextureBackedRenderTarget::CopyToInternalTexture()
 #if WITH_EDITOR
 	using namespace LandscapeTextureBackedRenderTargetLocals;
 
+	if (!IsCopyingBackAndForthAllowed())
+	{
+		return;
+	}
+
 	if (!ensure(IsValid(RenderTarget)))
 	{
 		return;
@@ -397,6 +412,11 @@ void ULandscapeHeightTextureBackedRenderTarget::CopyBackFromInternalTexture()
 #if WITH_EDITOR
 	using namespace LandscapeTextureBackedRenderTargetLocals;
 
+	if (!IsCopyingBackAndForthAllowed())
+	{
+		return;
+	}
+
 	if (!ensure(IsValid(InternalTexture)))
 	{
 		return;
@@ -465,4 +485,12 @@ void ULandscapeHeightTextureBackedRenderTarget::CopyBackFromInternalTexture()
 			});
 	}
 #endif // WITH_EDITOR
+}
+
+bool ULandscapeTextureBackedRenderTargetBase::IsCopyingBackAndForthAllowed()
+{
+	UWorld* World = GetWorld();
+	return IsValid(this) && !IsTemplate()
+		&& IsValid(World) && World->WorldType == EWorldType::Editor
+		&& FApp::CanEverRender();
 }
\ No newline at end of file
diff --git a/Engine/Plugins/Experimental/LandscapePatch/Source/LandscapePatch/Public/LandscapePatchManager.h b/Engine/Plugins/Experimental/LandscapePatch/Source/LandscapePatch/Public/LandscapePatchManager.h
--- a/Engine/Plugins/Experimental/LandscapePatch/Source/LandscapePatch/Public/LandscapePatchManager.h
+++ b/Engine/Plugins/Experimental/LandscapePatch/Source/LandscapePatch/Public/LandscapePatchManager.h
@@ -47,6 +47,8 @@ public:
 	 */
 	virtual FTransform GetHeightmapCoordsToWorld() { return HeightmapCoordsToWorld; }
 
+	bool ContainsPatch(TObjectPtr<ULandscapePatchComponent> Patch);
+
 	void AddPatch(TObjectPtr<ULandscapePatchComponent> Patch);
 
 	bool RemovePatch(TObjectPtr<ULandscapePatchComponent> Patch);
diff --git a/Engine/Plugins/Experimental/LandscapePatch/Source/LandscapePatch/Public/LandscapeTextureBackedRenderTarget.h b/Engine/Plugins/Experimental/LandscapePatch/Source/LandscapePatch/Public/LandscapeTextureBackedRenderTarget.h
--- a/Engine/Plugins/Experimental/LandscapePatch/Source/LandscapePatch/Public/LandscapeTextureBackedRenderTarget.h
+++ b/Engine/Plugins/Experimental/LandscapePatch/Source/LandscapePatch/Public/LandscapeTextureBackedRenderTarget.h
@@ -54,6 +54,10 @@ public:
 #endif // WITH_EDITOR
 
 protected:
+	// Allows us to disable copying back and forth between render target and texture when this is not possible
+	// (for instance, when rendering is turned off, so render target resources are not allocated).
+	virtual bool IsCopyingBackAndForthAllowed();
+
 	virtual ETextureSourceFormat GetInternalTextureFormat() PURE_VIRTUAL(ULandscapeTextureBackedRenderTargetBase::GetInternalTextureFormat, return ETextureSourceFormat::TSF_G8;);
 	virtual ETextureRenderTargetFormat GetRenderTargetFormat() PURE_VIRTUAL(ULandscapeTextureBackedRenderTargetBase::GetRenderTargetFormat, return ETextureRenderTargetFormat::RTF_R8;);
 
diff --git a/Engine/Plugins/Experimental/LandscapePatch/Source/LandscapePatch/Public/LandscapeTexturePatch.h b/Engine/Plugins/Experimental/LandscapePatch/Source/LandscapePatch/Public/LandscapeTexturePatch.h
--- a/Engine/Plugins/Experimental/LandscapePatch/Source/LandscapePatch/Public/LandscapeTexturePatch.h
+++ b/Engine/Plugins/Experimental/LandscapePatch/Source/LandscapePatch/Public/LandscapeTexturePatch.h
@@ -173,7 +173,9 @@ protected:
 	UPROPERTY(EditAnywhere, Category = WeightPatch, meta = (EditConditionHides, EditCondition = "false"))
 	ELandscapeTexturePatchSourceMode SourceMode = ELandscapeTexturePatchSourceMode::None;
 
-	UPROPERTY(EditAnywhere, Category = WeightPatch, meta = (DisplayName = "Source Mode"))
+	 //~ This is EditInstanceOnly because we can't create textures in the blueprint editor due to the way
+	 //~ instanced properties are currently handled there.
+	UPROPERTY(EditInstanceOnly, Category = WeightPatch, meta = (DisplayName = "Source Mode"))
 	ELandscapeTexturePatchSourceMode DetailPanelSourceMode = ELandscapeTexturePatchSourceMode::None;
 
 	//~ We could refactor things such that we always have an InternalData pointer, even when we use
@@ -455,9 +457,13 @@ protected:
 	ELandscapeTexturePatchSourceMode HeightSourceMode = ELandscapeTexturePatchSourceMode::None;
 
 	/**
-	 * How the heightmap of the patch is stored.
+	 * How the heightmap of the patch is stored. Not settable in the detail panel of the blueprint editor- use SetHeightSourceMode
+	 * in blueprint actors instead.
 	 */
-	UPROPERTY(EditAnywhere, Category = HeightPatch, meta = (DisplayName = "Source Mode"))
+	//~ This is EditInstanceOnly because changing it creates/destroys the internal texture, and that cannot currently be
+	//~ dealt with properly in the blueprint editor due to the way instanced properties are handled there. Could revisit when
+	//~ UE-158706 is resolved.
+	UPROPERTY(EditInstanceOnly, Category = HeightPatch, meta = (DisplayName = "Source Mode"))
 	ELandscapeTexturePatchSourceMode DetailPanelHeightSourceMode = ELandscapeTexturePatchSourceMode::None;
 
 	/** Not directly settable via detail panel- for display/debugging purposes only. */
@@ -502,8 +508,13 @@ protected:
 
 	// Weight properties:
 
-	/** */
-	UPROPERTY(EditAnywhere, Category = WeightPatches, Instanced)
+	/** 
+	 * Weight patches. These are not available to be manipulated through the detail panel in the blueprint editor, so blueprint actors
+	 * should set up the weight patches via AddWeightPatch instead.
+	 */
+	//~ This is EditInstanceOnly because manipulating them in blueprint editor causes saving issues due to the way that
+	//~ instanced properties are currently handled there.
+	UPROPERTY(EditInstanceOnly, Category = WeightPatches, Instanced)
 	TArray<TObjectPtr<ULandscapeWeightPatchTextureInfo>> WeightPatches;
 
 	// Used to detect changes to the number of weight patches via the detail panel, so that we can
