LandscapePatch: Keep patches at the same index in the patch manager list when rerunning owner construction script.

diff --git a/Engine/Plugins/Experimental/LandscapePatch/Source/LandscapePatch/Private/LandscapePatchComponent.cpp b/Engine/Plugins/Experimental/LandscapePatch/Source/LandscapePatch/Private/LandscapePatchComponent.cpp
--- a/Engine/Plugins/Experimental/LandscapePatch/Source/LandscapePatch/Private/LandscapePatchComponent.cpp
+++ b/Engine/Plugins/Experimental/LandscapePatch/Source/LandscapePatch/Private/LandscapePatchComponent.cpp
@@ -186,6 +186,21 @@ void ULandscapePatchComponent::GetActorDescProperties(FPropertyPairsMap& Propert
 	}
 }
 
+TStructOnScope<FActorComponentInstanceData> ULandscapePatchComponent::GetComponentInstanceData() const
+{
+	return MakeStructOnScope<FActorComponentInstanceData, FLandscapePatchComponentInstanceData>(this);
+}
+
+FLandscapePatchComponentInstanceData::FLandscapePatchComponentInstanceData(const ULandscapePatchComponent* Patch)
+: FSceneComponentInstanceData(Patch)
+{
+	ALandscapePatchManager* PatchManager = Patch->GetPatchManager();
+	if (PatchManager)
+	{
+		IndexInManager = PatchManager->GetIndexOfPatch(Patch);
+	}
+}
+
 void ULandscapePatchComponent::OnUpdateTransform(EUpdateTransformFlags UpdateTransformFlags, ETeleportType Teleport)
 {
 	Super::OnUpdateTransform(UpdateTransformFlags, Teleport);
@@ -223,6 +238,16 @@ void ULandscapePatchComponent::PostEditChangeProperty(FPropertyChangedEvent& Pro
 }
 #endif
 
+// Called after rerunning construction scripts (when patch is part of a blueprint) to make sure that the newly created 
+// instance of the patch is placed in the same order in the patch  manager as the destroyed instance.
+void ULandscapePatchComponent::ApplyComponentInstanceData(FLandscapePatchComponentInstanceData* ComponentInstanceData)
+{
+	if (ComponentInstanceData && ComponentInstanceData->IndexInManager >= 0 && PatchManager.IsValid())
+	{
+		PatchManager->MovePatchToIndex(this, ComponentInstanceData->IndexInManager);
+	}
+}
+
 void ULandscapePatchComponent::SetLandscape(ALandscape* NewLandscape)
 {
 #if WITH_EDITOR
@@ -280,10 +305,10 @@ void ULandscapePatchComponent::SetPatchManager(ALandscapePatchManager* NewPatchM
 	// surround everything since nothing is currently expected to work for a patch manager at runtime.
 #if WITH_EDITOR
 
-	// TODO: We don't currently have an early out here (for PreviousPatchManager == NewPatchManager) 
-	// because using SetPatchManager is currently a convenient way to reorder patches inside the manager,
-	// and because we want to make sure the patch is added to the manager if it was removed from
-	// the list somehow. However we may want to revisit this if it turns out to cause problems.
+	if (PreviousPatchManager == NewPatchManager && NewPatchManager && NewPatchManager->ContainsPatch(this))
+	{
+		return;
+	}
 
 	if (PreviousPatchManager.IsValid())
 	{
@@ -310,6 +335,22 @@ void ULandscapePatchComponent::SetPatchManager(ALandscapePatchManager* NewPatchM
 #endif // WITH_EDITOR
 }
 
+ALandscapePatchManager* ULandscapePatchComponent::GetPatchManager() const
+{
+	return PatchManager.Get();
+}
+
+void ULandscapePatchComponent::MoveToTop()
+{
+	if (PatchManager.IsValid() && PatchManager->ContainsPatch(this))
+	{
+		PatchManager->Modify();
+		PatchManager->RemovePatch(this);
+		PatchManager->AddPatch(this);
+		PatchManager->RequestLandscapeUpdate();
+	}
+}
+
 void ULandscapePatchComponent::RequestLandscapeUpdate()
 {
 	if(PatchManager.IsValid())
diff --git a/Engine/Plugins/Experimental/LandscapePatch/Source/LandscapePatch/Private/LandscapePatchManager.cpp b/Engine/Plugins/Experimental/LandscapePatch/Source/LandscapePatch/Private/LandscapePatchManager.cpp
--- a/Engine/Plugins/Experimental/LandscapePatch/Source/LandscapePatch/Private/LandscapePatchManager.cpp
+++ b/Engine/Plugins/Experimental/LandscapePatch/Source/LandscapePatch/Private/LandscapePatchManager.cpp
@@ -171,7 +171,7 @@ void ALandscapePatchManager::SetTargetLandscape(ALandscape* InTargetLandscape)
 #endif
 }
 
-bool ALandscapePatchManager::ContainsPatch(TObjectPtr<ULandscapePatchComponent> Patch)
+bool ALandscapePatchManager::ContainsPatch(TObjectPtr<ULandscapePatchComponent> Patch) const
 {
 	return PatchComponents.Contains(Patch);
 }
@@ -203,6 +203,29 @@ bool ALandscapePatchManager::RemovePatch(TObjectPtr<ULandscapePatchComponent> Pa
 	return bRemoved;
 }
 
+int32 ALandscapePatchManager::GetIndexOfPatch(TObjectPtr<const ULandscapePatchComponent> Patch) const
+{
+	return PatchComponents.IndexOfByKey(Patch);
+}
+
+void ALandscapePatchManager::MovePatchToIndex(TObjectPtr<ULandscapePatchComponent> Patch, int32 Index)
+{
+	if (!Patch)
+	{
+		return;
+	}
+
+	// It might seem like the index needs adjusting if we're removing before the given index, but that
+	// is not the case if our goal is for the index of the patch to be the given index at the end (rather
+	// than our goal being that the patch be in a particular position relative to the existing patches).
+	RemovePatch(Patch);
+
+	Index = FMath::Clamp(Index, 0, PatchComponents.Num());
+	PatchComponents.Insert(TSoftObjectPtr<ULandscapePatchComponent>(Patch.Get()), Index);
+
+	RequestLandscapeUpdate();
+}
+
 #if WITH_EDITOR
 bool ALandscapePatchManager::IsAffectingWeightmapLayer(const FName& InLayerName) const
 {
diff --git a/Engine/Plugins/Experimental/LandscapePatch/Source/LandscapePatch/Public/LandscapePatchComponent.h b/Engine/Plugins/Experimental/LandscapePatch/Source/LandscapePatch/Public/LandscapePatchComponent.h
--- a/Engine/Plugins/Experimental/LandscapePatch/Source/LandscapePatch/Public/LandscapePatchComponent.h
+++ b/Engine/Plugins/Experimental/LandscapePatch/Source/LandscapePatch/Public/LandscapePatchComponent.h
@@ -60,10 +60,16 @@ public:
 	UFUNCTION(BlueprintCallable, Category = "LandscapePatch")
 	virtual void SetPatchManager(ALandscapePatchManager* NewPatchManager);
 
+	UFUNCTION(BlueprintCallable, Category = "LandscapePatch")
+	virtual ALandscapePatchManager* GetPatchManager() const;
+
+	void ApplyComponentInstanceData(struct FLandscapePatchComponentInstanceData* ComponentInstanceData);
+
 	// For now we keep the patches largely editor-only, since we don't yet support runtime landscape editing.
 	// The above functions are also editor-only (and don't work at runtime), but can't be in WITH_EDITOR blocks
 	// so that they can be called from non-editor-only classes in editor contexts.
 #if WITH_EDITOR
+
 	// USceneComponent
 	virtual void OnUpdateTransform(EUpdateTransformFlags UpdateTransformFlags, ETeleportType Teleport) override;
 
@@ -72,6 +78,7 @@ public:
 	virtual void OnComponentDestroyed(bool bDestroyingHierarchy) override;
 	virtual void OnRegister() override;
 	virtual void GetActorDescProperties(FPropertyPairsMap& PropertyPairsMap) const override;
+	virtual TStructOnScope<FActorComponentInstanceData> GetComponentInstanceData() const override;
 
 	// UObject
 	virtual void PostEditChangeProperty(FPropertyChangedEvent& PropertyChangedEvent) override;
@@ -89,6 +96,13 @@ protected:
 	UPROPERTY(EditAnywhere, Category = Settings, AdvancedDisplay)
 	TSoftObjectPtr<ALandscapePatchManager> PatchManager = nullptr;
 
+	/** 
+	 * Move the patch to be the last processed patch in the current patch manager. This is a way to 
+	 * reorder patches relative to each other. 
+	 */
+	UFUNCTION(CallInEditor, Category = Initialization)
+	void MoveToTop();
+
 	// Determines whether the height patch was made by copying a different height patch.
 	bool bWasCopy = false;
 
@@ -111,4 +125,30 @@ private:
 	// Used to properly transition to a different manager when editing it via the detail panel.
 	UPROPERTY()
 	TSoftObjectPtr<ALandscapePatchManager> PreviousPatchManager = nullptr;
+};
+
+/** Used to store some extra data during RerunConstructionScripts, namely the component's position in the patch manager. */
+USTRUCT()
+struct FLandscapePatchComponentInstanceData : public FSceneComponentInstanceData
+{
+	GENERATED_BODY()
+
+	FLandscapePatchComponentInstanceData() = default;
+	FLandscapePatchComponentInstanceData(const ULandscapePatchComponent* SourceComponent);
+
+	virtual ~FLandscapePatchComponentInstanceData() = default;
+
+	virtual bool ContainsData() const override
+	{
+		return true;
+	}
+
+	virtual void ApplyToComponent(UActorComponent* Component, const ECacheApplyPhase CacheApplyPhase) override
+	{
+		Super::ApplyToComponent(Component, CacheApplyPhase);
+		CastChecked<ULandscapePatchComponent>(Component)->ApplyComponentInstanceData(this);
+	}
+
+	UPROPERTY()
+	int32 IndexInManager = -1;
 };
\ No newline at end of file
diff --git a/Engine/Plugins/Experimental/LandscapePatch/Source/LandscapePatch/Public/LandscapePatchManager.h b/Engine/Plugins/Experimental/LandscapePatch/Source/LandscapePatch/Public/LandscapePatchManager.h
--- a/Engine/Plugins/Experimental/LandscapePatch/Source/LandscapePatch/Public/LandscapePatchManager.h
+++ b/Engine/Plugins/Experimental/LandscapePatch/Source/LandscapePatch/Public/LandscapePatchManager.h
@@ -47,12 +47,23 @@ public:
 	 */
 	virtual FTransform GetHeightmapCoordsToWorld() { return HeightmapCoordsToWorld; }
 
-	bool ContainsPatch(TObjectPtr<ULandscapePatchComponent> Patch);
+	bool ContainsPatch(TObjectPtr<ULandscapePatchComponent> Patch) const;
 
 	void AddPatch(TObjectPtr<ULandscapePatchComponent> Patch);
 
 	bool RemovePatch(TObjectPtr<ULandscapePatchComponent> Patch);
 
+	/** 
+	 * Gets the index of a particular patch in the manager's stack of patches (later indices get applied after
+	 * earlier ones.
+	 */
+	int32 GetIndexOfPatch(TObjectPtr<const ULandscapePatchComponent> Patch) const;
+
+	/**
+	 * Moves patch to given index in the list of patches held by the manager (so that it is applied at 
+	 * a particular time relative to the others).
+	 */
+	void MovePatchToIndex(TObjectPtr<ULandscapePatchComponent> Patch, int32 Index);
 	
 #if WITH_EDITOR
 	// ALandscapeBlueprintBrushBase
