LandscapePatchManager: Remove ensure when finding a pending patch pointer that turns out to be invalid in the patch manager.

diff --git a/Engine/Plugins/Experimental/LandscapePatch/Source/LandscapePatch/Private/LandscapePatchManager.cpp b/Engine/Plugins/Experimental/LandscapePatch/Source/LandscapePatch/Private/LandscapePatchManager.cpp
--- a/Engine/Plugins/Experimental/LandscapePatch/Source/LandscapePatch/Private/LandscapePatchManager.cpp
+++ b/Engine/Plugins/Experimental/LandscapePatch/Source/LandscapePatch/Private/LandscapePatchManager.cpp
@@ -107,8 +107,14 @@ UTextureRenderTarget2D* ALandscapePatchManager::Render_Native(bool InIsHeightmap
 		}
 		else
 		{
-			// This means that IsPending() was true, but LoadSynchronous() failed, which we don't expect to happen.
-			ensure(false);
+			// This means that IsPending() was true, but LoadSynchronous() failed, which we generally don't
+			// expect to happen. However, it can happen in some edge cases such as if you force delete a patch
+			// holder blueprint and don't save the patch manager afterward. Whatever the reason, this is likely
+			// a dead patch that actually needs removal.
+			UE_LOG(LogLandscapePatch, Warning, TEXT("ALandscapePatchManager: Found a pending patch pointer in patch manager that "
+				"turned out to be invalid. It will be removed."));
+			Component = nullptr;
+			bHaveInvalidPatches = true;
 		}
 	}
 
