LandscapePatch: Remove the old LandscapeTextureHeightPatch.

diff --git a/Engine/Plugins/Experimental/LandscapePatch/Source/LandscapePatch/Private/LandscapeTextureHeightPatch.cpp b/Engine/Plugins/Experimental/LandscapePatch/Source/LandscapePatch/Private/LandscapeTextureHeightPatch.cpp
deleted file mode 100644
diff --git a/Engine/Plugins/Experimental/LandscapePatch/Source/LandscapePatch/Private/LandscapeTexturePatchBase.cpp b/Engine/Plugins/Experimental/LandscapePatch/Source/LandscapePatch/Private/LandscapeTexturePatchBase.cpp
deleted file mode 100644
diff --git a/Engine/Plugins/Experimental/LandscapePatch/Source/LandscapePatch/Public/LandscapeTextureHeightPatch.h b/Engine/Plugins/Experimental/LandscapePatch/Source/LandscapePatch/Public/LandscapeTextureHeightPatch.h
deleted file mode 100644
diff --git a/Engine/Plugins/Experimental/LandscapePatch/Source/LandscapePatch/Public/LandscapeTexturePatchBase.h b/Engine/Plugins/Experimental/LandscapePatch/Source/LandscapePatch/Public/LandscapeTexturePatchBase.h
deleted file mode 100644
diff --git a/Engine/Plugins/Experimental/LandscapePatch/Source/LandscapePatchEditorOnly/Private/LandscapePatchEditorOnlyModule.cpp b/Engine/Plugins/Experimental/LandscapePatch/Source/LandscapePatchEditorOnly/Private/LandscapePatchEditorOnlyModule.cpp
--- a/Engine/Plugins/Experimental/LandscapePatch/Source/LandscapePatchEditorOnly/Private/LandscapePatchEditorOnlyModule.cpp
+++ b/Engine/Plugins/Experimental/LandscapePatch/Source/LandscapePatchEditorOnly/Private/LandscapePatchEditorOnlyModule.cpp
@@ -4,8 +4,6 @@
 
 #include "Editor/UnrealEdEngine.h"
 #include "LandscapeTexturePatch.h"
-#include "LandscapeTexturePatchBase.h"
-#include "LandscapeTexturePatchCustomization.h"
 #include "LandscapeTexturePatchVisualizer.h"
 #include "PropertyEditorModule.h"
 #include "UnrealEdGlobals.h"
@@ -17,22 +15,11 @@ void FLandscapePatchEditorOnlyModule::StartupModule()
 	if (GUnrealEd)
 	{
 		TSharedPtr<FLandscapeTexturePatchVisualizer> Visualizer = MakeShared<FLandscapeTexturePatchVisualizer>();
-		GUnrealEd->RegisterComponentVisualizer(ULandscapeTexturePatchBase::StaticClass()->GetFName(), Visualizer);
 		GUnrealEd->RegisterComponentVisualizer(ULandscapeTexturePatch::StaticClass()->GetFName(), Visualizer);
 		// This call should maybe be inside the RegisterComponentVisualizer call above, but since it's not,
 		// we'll put it here.
 		Visualizer->OnRegister();
-		VisualizersToUnregisterOnShutdown.Add(ULandscapeTexturePatchBase::StaticClass()->GetFName());
-	}
-
-	// Register detail customization
-	ClassesToUnregisterOnShutdown.Reset();
-	FPropertyEditorModule* PropertyEditorModule = FModuleManager::GetModulePtr<FPropertyEditorModule>("PropertyEditor");
-	if (PropertyEditorModule)
-	{
-		PropertyEditorModule->RegisterCustomClassLayout(ULandscapeTexturePatchBase::StaticClass()->GetFName(),
-			FOnGetDetailCustomizationInstance::CreateStatic(&FLandscapeTexturePatchCustomization::MakeInstance));
-		ClassesToUnregisterOnShutdown.Add(ULandscapeTexturePatchBase::StaticClass()->GetFName());
+		VisualizersToUnregisterOnShutdown.Add(ULandscapeTexturePatch::StaticClass()->GetFName());
 	}
 }
 
diff --git a/Engine/Plugins/Experimental/LandscapePatch/Source/LandscapePatchEditorOnly/Private/LandscapeTexturePatchCustomization.cpp b/Engine/Plugins/Experimental/LandscapePatch/Source/LandscapePatchEditorOnly/Private/LandscapeTexturePatchCustomization.cpp
deleted file mode 100644
diff --git a/Engine/Plugins/Experimental/LandscapePatch/Source/LandscapePatchEditorOnly/Private/LandscapeTexturePatchVisualizer.cpp b/Engine/Plugins/Experimental/LandscapePatch/Source/LandscapePatchEditorOnly/Private/LandscapeTexturePatchVisualizer.cpp
--- a/Engine/Plugins/Experimental/LandscapePatch/Source/LandscapePatchEditorOnly/Private/LandscapeTexturePatchVisualizer.cpp
+++ b/Engine/Plugins/Experimental/LandscapePatch/Source/LandscapePatchEditorOnly/Private/LandscapeTexturePatchVisualizer.cpp
@@ -3,7 +3,6 @@
 #include "LandscapeTexturePatchVisualizer.h"
 
 #include "Landscape.h"
-#include "LandscapeTexturePatchBase.h"
 #include "LandscapeTexturePatch.h"
 
 #include "SceneManagement.h" // FPrimitiveDrawInterface
@@ -22,14 +21,6 @@ void FLandscapeTexturePatchVisualizer::DrawVisualization(const UActorComponent*
 			Color, Patch->GetUnscaledCoverage().X * PatchToWorld.GetScale3D().X, Patch->GetUnscaledCoverage().Y * PatchToWorld.GetScale3D().Y, SDPG_Foreground,
 			Thickness, DepthBias, bScreenSpace);
 	}
-	// TODO: Should be removed once we delete ULandscapeTexturePatchBase and ULandscapeTextureHeightPatch
-	else if (const ULandscapeTexturePatchBase* HeightPatch = Cast<ULandscapeTexturePatchBase>(Component))
-	{
-		FTransform PatchToWorld = HeightPatch->GetPatchToWorldTransform();
-		DrawRectangle(PDI, PatchToWorld.GetTranslation(), PatchToWorld.GetUnitAxis(EAxis::X), PatchToWorld.GetUnitAxis(EAxis::Y),
-			Color, HeightPatch->GetUnscaledCoverage().X * PatchToWorld.GetScale3D().X, HeightPatch->GetUnscaledCoverage().Y * PatchToWorld.GetScale3D().Y, SDPG_Foreground,
-			Thickness, DepthBias, bScreenSpace);
-	}
 	else
 	{
 		ensure(false);
diff --git a/Engine/Plugins/Experimental/LandscapePatch/Source/LandscapePatchEditorOnly/Public/LandscapeTexturePatchCustomization.h b/Engine/Plugins/Experimental/LandscapePatch/Source/LandscapePatchEditorOnly/Public/LandscapeTexturePatchCustomization.h
deleted file mode 100644
