LandscapePatch: Preserve current state of render targets across construction script reruns. Consider texture patches to be disabled if all source modes are "none", and don't request landscape updates while disabled. Give ability to reattach a patch manager to a landscape. Give warnings for issues that might arise if landscape or patch manager are not saved.

diff --git a/Engine/Plugins/Experimental/LandscapePatch/Source/LandscapePatch/Private/LandscapePatchComponent.cpp b/Engine/Plugins/Experimental/LandscapePatch/Source/LandscapePatch/Private/LandscapePatchComponent.cpp
--- a/Engine/Plugins/Experimental/LandscapePatch/Source/LandscapePatch/Private/LandscapePatchComponent.cpp
+++ b/Engine/Plugins/Experimental/LandscapePatch/Source/LandscapePatch/Private/LandscapePatchComponent.cpp
@@ -173,7 +173,10 @@ void ULandscapePatchComponent::OnRegister()
 	// variety of changes this way, so we'll need to see if we can get rid of other invalidations.
 	// Also, we should make the invalidation conditional on whether we actually modify any relevant
 	// properties by having a virtual method that compares and updates a stored hash of them.
-	RequestLandscapeUpdate();
+	if (IsEnabled())
+	{
+		RequestLandscapeUpdate();
+	}
 }
 
 void ULandscapePatchComponent::GetActorDescProperties(FPropertyPairsMap& PropertyPairsMap) const
@@ -198,6 +201,12 @@ FLandscapePatchComponentInstanceData::FLandscapePatchComponentInstanceData(const
 	if (PatchManager)
 	{
 		IndexInManager = PatchManager->GetIndexOfPatch(Patch);
+
+#if WITH_EDITOR
+		bGaveMissingPatchManagerWarning = Patch->bGaveMissingPatchManagerWarning;
+		bGaveNotInPatchManagerWarning = Patch->bGaveNotInPatchManagerWarning;
+		bGaveMissingLandscapeWarning = Patch->bGaveMissingLandscapeWarning;
+#endif
 	}
 }
 
@@ -205,15 +214,16 @@ void ULandscapePatchComponent::OnUpdateTransform(EUpdateTransformFlags UpdateTra
 {
 	Super::OnUpdateTransform(UpdateTransformFlags, Teleport);
 
-	RequestLandscapeUpdate();
+	if (IsEnabled())
+	{
+		RequestLandscapeUpdate();
+	}
 }
 
 void ULandscapePatchComponent::PostEditChangeProperty(FPropertyChangedEvent& PropertyChangedEvent)
 {
 	using namespace LandscapePatchComponentLocals;
 
-	Super::PostEditChangeProperty(PropertyChangedEvent);
-
 	// Do a bunch of checks to make sure that we don't try to do anything when the editing is happening inside the blueprint editor.
 	UWorld* World = GetWorld();
 	if (IsTemplate() || !IsValid(this) || !IsValid(World) || World->WorldType != EWorldType::Editor)
@@ -234,7 +244,17 @@ void ULandscapePatchComponent::PostEditChangeProperty(FPropertyChangedEvent& Pro
 		SetPatchManager(PatchManager.Get());
 	}
 
-	RequestLandscapeUpdate();
+	if (IsEnabled() || (PropertyChangedEvent.Property
+		&& (PropertyChangedEvent.Property->GetFName() == GET_MEMBER_NAME_CHECKED(ULandscapePatchComponent, bIsEnabled))))
+	{
+		RequestLandscapeUpdate();
+	}
+
+	// It is important that this super call happen after the above, because inside a blueprint actor, the call triggers a
+	// rerun of the construction scripts, which will destroy the component and mess with our ability to do the above adjustments
+	// properly (IsValid(this) returns false, the patch manager has the patch removed so it complains when we try to trigger
+	// the update, etc).
+	Super::PostEditChangeProperty(PropertyChangedEvent);
 }
 #endif
 
@@ -246,6 +266,12 @@ void ULandscapePatchComponent::ApplyComponentInstanceData(FLandscapePatchCompone
 	{
 		PatchManager->MovePatchToIndex(this, ComponentInstanceData->IndexInManager);
 	}
+
+#if WITH_EDITOR
+	bGaveMissingPatchManagerWarning = ComponentInstanceData->bGaveMissingPatchManagerWarning;
+	bGaveNotInPatchManagerWarning = ComponentInstanceData->bGaveNotInPatchManagerWarning;
+	bGaveMissingLandscapeWarning = ComponentInstanceData->bGaveMissingLandscapeWarning;
+#endif
 }
 
 void ULandscapePatchComponent::SetLandscape(ALandscape* NewLandscape)
@@ -254,10 +280,8 @@ void ULandscapePatchComponent::SetLandscape(ALandscape* NewLandscape)
 
 	using namespace LandscapePatchComponentLocals;
 
-	// Uncertain whether we want this early-out here. Perhaps, like SetPatchManager, we want 
-	// to be able to use this function to reorder patches. On the other hand, it seems inconvenient
-	// to accidentally swap patch managers if there are multiple in the same landscape, and we kept
-	// landscape the same. It's hard to know the ideal behavior, but for now we'll keep it.
+	// We early out here if the landscape stays the same and we have a patch manager because it is inconvenient to accidentally swap patch
+	// managers if there are multiple in the same landscape.
 	if (Landscape.Get() == NewLandscape && (
 		(Landscape.IsNull() && PatchManager.IsNull()) || (Landscape.IsValid() && PatchManager.IsValid() && PatchManager->GetOwningLandscape() == Landscape.Get())))
 	{
@@ -309,6 +333,7 @@ void ULandscapePatchComponent::SetPatchManager(ALandscapePatchManager* NewPatchM
 	{
 		return;
 	}
+	ResetWarnings();
 
 	if (PreviousPatchManager.IsValid())
 	{
@@ -347,18 +372,66 @@ void ULandscapePatchComponent::MoveToTop()
 		PatchManager->Modify();
 		PatchManager->RemovePatch(this);
 		PatchManager->AddPatch(this);
-		PatchManager->RequestLandscapeUpdate();
+		if (IsEnabled())
+		{
+			PatchManager->RequestLandscapeUpdate();
+		}
 	}
 }
 
 void ULandscapePatchComponent::RequestLandscapeUpdate()
 {
-	if(PatchManager.IsValid())
+#if WITH_EDITOR
+	if (!PatchManager.IsValid())
+	{
+		if (!bGaveMissingPatchManagerWarning)
+		{
+			UE_LOG(LogLandscapePatch, Warning, TEXT("Patch does not have a valid patch manager. "
+				"Set the landscape or patch manager on the patch."));
+			bGaveMissingPatchManagerWarning = true;
+		}
+		return;
+	}
+
+	bool bRequestUpdate = true;
+	if (!PatchManager->ContainsPatch(this))
+	{
+		if (!bGaveNotInPatchManagerWarning)
+		{
+			UE_LOG(LogLandscapePatch, Warning, TEXT("Patch's patch manager does not contain this patch. "
+				"Perhaps the manager was not saved? Reset the patch manager on the patch."));
+			bGaveNotInPatchManagerWarning = true;
+		}
+		bRequestUpdate = false;
+	}
+	if (!IsValid(PatchManager->GetOwningLandscape()))
 	{
+		if (!bGaveMissingLandscapeWarning)
+		{
+			UE_LOG(LogLandscapePatch, Warning, TEXT("Patch's patch manager does not have a valid owning "
+				"landscape. Perhaps the landscape was not saved? Reset the landscape on the manager."));
+			bGaveMissingLandscapeWarning = true;
+		}
+		bRequestUpdate = false;
+	}
+
+	if(bRequestUpdate)
+	{
+		ResetWarnings();
 		PatchManager->RequestLandscapeUpdate();
 	}
+#endif
 }
 
+#if WITH_EDITOR
+void ULandscapePatchComponent::ResetWarnings()
+{
+	bGaveMissingPatchManagerWarning = false;
+	bGaveNotInPatchManagerWarning = false;
+	bGaveMissingLandscapeWarning = false;
+}
+#endif
+
 FTransform ULandscapePatchComponent::GetLandscapeHeightmapCoordsToWorld() const
 {
 	if (PatchManager.IsValid())
diff --git a/Engine/Plugins/Experimental/LandscapePatch/Source/LandscapePatch/Private/LandscapePatchManager.cpp b/Engine/Plugins/Experimental/LandscapePatch/Source/LandscapePatch/Private/LandscapePatchManager.cpp
--- a/Engine/Plugins/Experimental/LandscapePatch/Source/LandscapePatch/Private/LandscapePatchManager.cpp
+++ b/Engine/Plugins/Experimental/LandscapePatch/Source/LandscapePatch/Private/LandscapePatchManager.cpp
@@ -145,7 +145,17 @@ void ALandscapePatchManager::SetTargetLandscape(ALandscape* InTargetLandscape)
 			ILandscapeModule& LandscapeModule = FModuleManager::GetModuleChecked<ILandscapeModule>("Landscape");
 			int32 PatchLayerIndex = LandscapeModule.GetLandscapeEditorServices()->GetOrCreateEditLayer(PatchLayerName, InTargetLandscape);
 			
+			// Among other things, this will call SetOwningLandscape on us.
 			InTargetLandscape->AddBrushToLayer(PatchLayerIndex, this);
+
+			// It's not clear whether this is really necessary, but we do it for consistency because Landscape does this in its
+			// PostLoad for all its brushes (through FLandscapeLayerBrush::SetOwner). One would think that it would be done 
+			// in AddBrushToLayer if it were at all important, but it currently isn't...
+			if (this->GetTypedOuter<ULevel>() != InTargetLandscape->GetTypedOuter<ULevel>())
+			{
+				// Change owner to be that level
+				this->Rename(nullptr, InTargetLandscape->GetTypedOuter<ULevel>());
+			}
 		}
 	}
 #endif
@@ -162,7 +172,13 @@ void ALandscapePatchManager::AddPatch(TObjectPtr<ULandscapePatchComponent> Patch
 	{
 		Modify();
 		PatchComponents.AddUnique(TSoftObjectPtr<ULandscapePatchComponent>(Patch.Get()));
-		RequestLandscapeUpdate();
+
+		// No need to update if the patch is disabled. Important to avoid needlessly updating while dragging a blueprint with
+		// a disabled patch (since construction scripts constantly add and remove).
+		if (Patch->IsEnabled())
+		{
+			RequestLandscapeUpdate();
+		}
 	}
 }
 
@@ -174,7 +190,10 @@ bool ALandscapePatchManager::RemovePatch(TObjectPtr<ULandscapePatchComponent> Pa
 	{
 		Modify();
 		bRemoved = PatchComponents.Remove(TSoftObjectPtr<ULandscapePatchComponent>(Patch.Get())) > 0;
-		if (bRemoved)
+
+		// No need to update if the patch was already disabled.Important to avoid needlessly updating while dragging 
+		// a blueprint with a disabled patch (since construction scripts constantly add and remove).
+		if (bRemoved && Patch->IsEnabled())
 		{
 			RequestLandscapeUpdate();
 		}
@@ -190,7 +209,7 @@ int32 ALandscapePatchManager::GetIndexOfPatch(TObjectPtr<const ULandscapePatchCo
 
 void ALandscapePatchManager::MovePatchToIndex(TObjectPtr<ULandscapePatchComponent> Patch, int32 Index)
 {
-	if (!Patch)
+	if (!Patch || Index < 0 || GetIndexOfPatch(Patch) == Index)
 	{
 		return;
 	}
@@ -203,7 +222,10 @@ void ALandscapePatchManager::MovePatchToIndex(TObjectPtr<ULandscapePatchComponen
 	Index = FMath::Clamp(Index, 0, PatchComponents.Num());
 	PatchComponents.Insert(TSoftObjectPtr<ULandscapePatchComponent>(Patch.Get()), Index);
 
-	RequestLandscapeUpdate();
+	if (Patch->IsEnabled())
+	{
+		RequestLandscapeUpdate();
+	}
 }
 
 #if WITH_EDITOR
@@ -229,6 +251,32 @@ void ALandscapePatchManager::PostEditUndo()
 {
 	RequestLandscapeUpdate();
 }
+
+void ALandscapePatchManager::SetOwningLandscape(ALandscape* InOwningLandscape)
+{
+	Super::SetOwningLandscape(InOwningLandscape);
+
+	DetailPanelLandscape = OwningLandscape;
+}
+
+// We override PostEditChange to allow the users to change the owning landscape via a property displayed in the detail panel.
+void ALandscapePatchManager::PostEditChangeProperty(FPropertyChangedEvent& PropertyChangedEvent)
+{
+	Super::PostEditChangeProperty(PropertyChangedEvent);
+
+	// Do a bunch of checks to make sure that we don't try to do anything when the editing is happening inside the blueprint editor.
+	UWorld* World = GetWorld();
+	if (IsTemplate() || !IsValid(this) || !IsValid(World) || World->WorldType != EWorldType::Editor)
+	{
+		return;
+	}
+
+	if (PropertyChangedEvent.Property
+		&& (PropertyChangedEvent.Property->GetFName() == GET_MEMBER_NAME_CHECKED(ALandscapePatchManager, DetailPanelLandscape)))
+	{
+		SetTargetLandscape(DetailPanelLandscape.Get());
+	}
+}
 #endif
 
 #undef LOCTEXT_NAMESPACE
\ No newline at end of file
diff --git a/Engine/Plugins/Experimental/LandscapePatch/Source/LandscapePatch/Private/LandscapeTextureBackedRenderTarget.cpp b/Engine/Plugins/Experimental/LandscapePatch/Source/LandscapePatch/Private/LandscapeTextureBackedRenderTarget.cpp
--- a/Engine/Plugins/Experimental/LandscapePatch/Source/LandscapePatch/Private/LandscapeTextureBackedRenderTarget.cpp
+++ b/Engine/Plugins/Experimental/LandscapePatch/Source/LandscapePatch/Private/LandscapeTextureBackedRenderTarget.cpp
@@ -39,6 +39,7 @@ void ULandscapeTextureBackedRenderTargetBase::PreSave(FObjectPreSaveContext Save
 	}
 }
 
+// This gets called not just when loading, but also after duplication.
 void ULandscapeTextureBackedRenderTargetBase::PostLoad()
 {
 	Super::PostLoad();
@@ -56,6 +57,7 @@ void ULandscapeTextureBackedRenderTargetBase::PostLoad()
 	}
 }
 
+// PreDuplicate is called not just when objects are copied in editor, but also when they are serialized for undo/redo.
 void ULandscapeTextureBackedRenderTargetBase::PreDuplicate(FObjectDuplicationParameters& DupParams)
 {
 	Super::PreDuplicate(DupParams);
@@ -77,6 +79,7 @@ void ULandscapeTextureBackedRenderTargetBase::ExportCustomProperties(FOutputDevi
 	}
 }
 
+// Called after pasting
 void ULandscapeTextureBackedRenderTargetBase::PostEditImport()
 {
 	Super::PostEditImport();
@@ -491,6 +494,8 @@ bool ULandscapeTextureBackedRenderTargetBase::IsCopyingBackAndForthAllowed()
 {
 	UWorld* World = GetWorld();
 	return IsValid(this) && !IsTemplate()
-		&& IsValid(World) && World->WorldType == EWorldType::Editor
+		// Note that having a null world is ok because we get that temporarily while rerunning construction
+		// scripts. However if we do have a world, it should be the normal editor one.
+		&& (!World || (IsValid(World) && World->WorldType == EWorldType::Editor))
 		&& FApp::CanEverRender();
 }
\ No newline at end of file
diff --git a/Engine/Plugins/Experimental/LandscapePatch/Source/LandscapePatch/Private/LandscapeTexturePatch.cpp b/Engine/Plugins/Experimental/LandscapePatch/Source/LandscapePatch/Private/LandscapeTexturePatch.cpp
--- a/Engine/Plugins/Experimental/LandscapePatch/Source/LandscapePatch/Private/LandscapeTexturePatch.cpp
+++ b/Engine/Plugins/Experimental/LandscapePatch/Source/LandscapePatch/Private/LandscapeTexturePatch.cpp
@@ -494,6 +494,12 @@ FLandscapeHeightPatchConvertToNativeParams ULandscapeTexturePatch::GetHeightConv
 void ULandscapeTexturePatch::ReinitializeHeight()
 {
 #if WITH_EDITOR
+	if (!Super::IsEnabled())
+	{
+		UE_LOG(LogLandscapePatch, Warning, TEXT("ULandscapeTexturePatch::Reinitialize: Cannot reinitialize while disabled."));
+		return;
+	}
+
 	if (!Landscape.IsValid() || !PatchManager.IsValid())
 	{
 		UE_LOG(LogLandscapePatch, Warning, TEXT("ULandscapeTexturePatch::Reinitialize: No associated landscape to initialize from."));
@@ -516,6 +522,12 @@ void ULandscapeTexturePatch::ReinitializeHeight()
 void ULandscapeTexturePatch::ReinitializeWeights()
 {
 #if WITH_EDITOR
+	if (!Super::IsEnabled())
+	{
+		UE_LOG(LogLandscapePatch, Warning, TEXT("ULandscapeTexturePatch::Reinitialize: Cannot reinitialize while disabled."));
+		return;
+	}
+
 	if (!Landscape.IsValid() || !PatchManager.IsValid())
 	{
 		UE_LOG(LogLandscapePatch, Warning, TEXT("ULandscapeTexturePatch::Reinitialize: No associated landscape to initialize from."));
@@ -675,11 +687,15 @@ void ULandscapeTexturePatch::ReinitializeHeight(UTextureRenderTarget2D* InCombin
 		GraphBuilder.Execute();
 	});
 
+	// The Modify() calls currently don't really help because we don't transact inside Render_Native. Maybe someday
+	// we'll add that ability (though it sounds messy).
+	HeightInternalData->GetInternalTexture()->Modify();
 	TemporaryNativeHeightCopy->UpdateTexture2D(HeightInternalData->GetInternalTexture(), ETextureSourceFormat::TSF_BGRA8);
 	HeightInternalData->GetInternalTexture()->UpdateResource();
 
 	if (IsValid(HeightInternalData->GetRenderTarget()))
 	{
+		HeightInternalData->GetRenderTarget()->Modify();
 		HeightInternalData->CopyBackFromInternalTexture();
 	}
 }
@@ -793,6 +809,29 @@ bool ULandscapeTexturePatch::IsAffectingWeightmapLayer(const FName& InLayerName)
 	}
 	return false;
 }
+
+// We override IsEnabled to make the patch not request updates when all the source modes are "none"
+// (unless we need the update for reinitialization).
+bool ULandscapeTexturePatch::IsEnabled() const
+{
+	if (!Super::IsEnabled())
+	{
+		return false;
+	}
+	if (HeightSourceMode != ELandscapeTexturePatchSourceMode::None || bReinitializeHeightOnNextRender)
+	{
+		return true;
+	}
+	for (TObjectPtr<ULandscapeWeightPatchTextureInfo> WeightPatch : WeightPatches)
+	{
+		if (WeightPatch->SourceMode != ELandscapeTexturePatchSourceMode::None || WeightPatch->bReinitializeOnNextRender)
+		{
+			return true;
+		}
+	}
+	// If we got to here, we are enabled, but all of the contained data had a source mode of "None"
+	return false;
+}
 #endif // WITH_EDITOR
 
 void ULandscapeTexturePatch::SnapToLandscape()
@@ -915,7 +954,7 @@ void ULandscapeTexturePatch::SetResolution(FVector2D ResolutionIn)
 		{
 			return;
 		}
-		else
+		else if (ensure(IsValid(InternalData)))
 		{
 			InternalData->SetResolution(DesiredX, DesiredY);
 		}
@@ -924,7 +963,10 @@ void ULandscapeTexturePatch::SetResolution(FVector2D ResolutionIn)
 	ResizePatch(HeightSourceMode, HeightInternalData);
 	for (TObjectPtr<ULandscapeWeightPatchTextureInfo> WeightPatch : WeightPatches)
 	{
-		ResizePatch(WeightPatch->SourceMode, WeightPatch->InternalData);
+		if (ensure(IsValid(WeightPatch)))
+		{
+			ResizePatch(WeightPatch->SourceMode, WeightPatch->InternalData);
+		}
 	}
 }
 
@@ -1036,6 +1078,20 @@ void ULandscapeWeightPatchTextureInfo::PostEditChangeProperty(FPropertyChangedEv
 
 	Super::PostEditChangeProperty(PropertyChangedEvent);
 }
+
+void ULandscapeWeightPatchTextureInfo::PreDuplicate(FObjectDuplicationParameters& DupParams)
+{
+	// TODO: It seems like this whole overload shouldn't be necessary, because we should get PreDuplicate calls
+	// on InternalData. However for reasons that I have yet to undertand, those calls are not made. It seems like
+	// there is different behavior for an array of instanced classes containing instanced properties...
+
+	Super::PreDuplicate(DupParams);
+
+	if (SourceMode == ELandscapeTexturePatchSourceMode::TextureBackedRenderTarget)
+	{
+		InternalData->CopyToInternalTexture();
+	}
+}
 #endif // WITH_EDITOR
 
 void ULandscapeWeightPatchTextureInfo::SetSourceMode(ELandscapeTexturePatchSourceMode NewMode)
diff --git a/Engine/Plugins/Experimental/LandscapePatch/Source/LandscapePatch/Public/LandscapePatchComponent.h b/Engine/Plugins/Experimental/LandscapePatch/Source/LandscapePatch/Public/LandscapePatchComponent.h
--- a/Engine/Plugins/Experimental/LandscapePatch/Source/LandscapePatch/Public/LandscapePatchComponent.h
+++ b/Engine/Plugins/Experimental/LandscapePatch/Source/LandscapePatch/Public/LandscapePatchComponent.h
@@ -49,7 +49,7 @@ public:
 	 * @return false if the patch is marked as disabled and therefore can't affect the landscape.
 	 */
 	UFUNCTION(BlueprintCallable, Category = "LandscapePatch")
-	bool IsEnabled() const { return bIsEnabled; }
+	virtual bool IsEnabled() const { return bIsEnabled; }
 
 	UFUNCTION(BlueprintCallable, Category = "LandscapePatch")
 	FTransform GetLandscapeHeightmapCoordsToWorld() const;
@@ -90,18 +90,24 @@ public:
 	virtual bool NeedsLoadForServer() const override { return false; }
 
 protected:
+	/**
+	 * Move the patch to be the last processed patch in the current patch manager. This is a way to
+	 * reorder patches relative to each other.
+	 */
+	UFUNCTION(CallInEditor, Category = Initialization)
+	void MoveToTop();
+
 	UPROPERTY(EditAnywhere, Category = Settings)
 	TSoftObjectPtr<ALandscape> Landscape = nullptr;
 
 	UPROPERTY(EditAnywhere, Category = Settings, AdvancedDisplay)
 	TSoftObjectPtr<ALandscapePatchManager> PatchManager = nullptr;
 
-	/** 
-	 * Move the patch to be the last processed patch in the current patch manager. This is a way to 
-	 * reorder patches relative to each other. 
+	/**
+	 * When false, patch does not affect the landscape. Useful for temporarily disabling the patch.
 	 */
-	UFUNCTION(CallInEditor, Category = Initialization)
-	void MoveToTop();
+	UPROPERTY(EditAnywhere, Category = Settings)
+	bool bIsEnabled = true;
 
 	// Determines whether the height patch was made by copying a different height patch.
 	bool bWasCopy = false;
@@ -110,12 +116,6 @@ protected:
 	// the first OnRegister call. It remains false from the first OnRegiter call onward, even
 	// if the component is unregistered.
 	bool bLoadedButNotYetRegistered = false;
-
-	/** 
-	 * When false, patch does not affect the landscape. Useful for temporarily disabling the patch. 
-	 */
-	UPROPERTY(EditAnywhere, Category = Settings)
-	bool bIsEnabled = true;
 private:
 	// Starts as false and gets set to true in construction, so gets used to set bWasCopy
 	// by checking the indicator value at the start of construction.
@@ -125,6 +125,16 @@ private:
 	// Used to properly transition to a different manager when editing it via the detail panel.
 	UPROPERTY()
 	TSoftObjectPtr<ALandscapePatchManager> PreviousPatchManager = nullptr;
+
+#if WITH_EDITOR
+	// Used to avoid spamming warning messages
+	bool bGaveMissingPatchManagerWarning = false;
+	bool bGaveNotInPatchManagerWarning = false;
+	bool bGaveMissingLandscapeWarning = false;
+	void ResetWarnings();
+#endif
+
+	friend struct FLandscapePatchComponentInstanceData;
 };
 
 /** Used to store some extra data during RerunConstructionScripts, namely the component's position in the patch manager. */
@@ -151,4 +161,16 @@ struct FLandscapePatchComponentInstanceData : public FSceneComponentInstanceData
 
 	UPROPERTY()
 	int32 IndexInManager = -1;
+
+#if WITH_EDITORONLY_DATA
+
+	// Used so that we don't spam warning messages while rerunning construction scripts on a patch
+	// that triggers one of the warnings.
+	UPROPERTY()
+	bool bGaveMissingPatchManagerWarning = false;
+	UPROPERTY()
+	bool bGaveNotInPatchManagerWarning = false;
+	UPROPERTY()
+	bool bGaveMissingLandscapeWarning = false;
+#endif
 };
\ No newline at end of file
diff --git a/Engine/Plugins/Experimental/LandscapePatch/Source/LandscapePatch/Public/LandscapePatchManager.h b/Engine/Plugins/Experimental/LandscapePatch/Source/LandscapePatch/Public/LandscapePatchManager.h
--- a/Engine/Plugins/Experimental/LandscapePatch/Source/LandscapePatch/Public/LandscapePatchManager.h
+++ b/Engine/Plugins/Experimental/LandscapePatch/Source/LandscapePatch/Public/LandscapePatchManager.h
@@ -68,9 +68,11 @@ public:
 #if WITH_EDITOR
 	// ALandscapeBlueprintBrushBase
 	virtual bool IsAffectingWeightmapLayer(const FName& InLayerName) const override;
+	virtual void SetOwningLandscape(class ALandscape* InOwningLandscape) override;
 
 	// UObject
 	virtual void PostEditUndo() override;
+	virtual void PostEditChangeProperty(FPropertyChangedEvent& PropertyChangedEvent) override;
 #endif
 	virtual bool IsEditorOnly() const override { return true; }
 	virtual bool NeedsLoadForClient() const override { return false; }
@@ -84,4 +86,12 @@ protected:
 	UPROPERTY()
 	FTransform HeightmapCoordsToWorld;
 
+#if WITH_EDITORONLY_DATA
+	//~ This is transient because SetOwningLandscape is called in ALandscape::PostLoad.
+	/**
+	 * The owning landscape.
+	 */
+	UPROPERTY(EditAnywhere, Category = Landscape, Transient, meta = (DisplayName = "Landscape"))
+	TObjectPtr<ALandscape> DetailPanelLandscape = nullptr;
+#endif
 };
\ No newline at end of file
diff --git a/Engine/Plugins/Experimental/LandscapePatch/Source/LandscapePatch/Public/LandscapeTexturePatch.h b/Engine/Plugins/Experimental/LandscapePatch/Source/LandscapePatch/Public/LandscapeTexturePatch.h
--- a/Engine/Plugins/Experimental/LandscapePatch/Source/LandscapePatch/Public/LandscapeTexturePatch.h
+++ b/Engine/Plugins/Experimental/LandscapePatch/Source/LandscapePatch/Public/LandscapeTexturePatch.h
@@ -154,6 +154,7 @@ public:
 #if WITH_EDITOR
 	// UObject
 	virtual void PostEditChangeProperty(FPropertyChangedEvent& PropertyChangedEvent) override;
+	virtual void PreDuplicate(FObjectDuplicationParameters& DupParams) override;
 #endif // WITH_EDITOR
 
 protected:
@@ -220,6 +221,7 @@ public:
 
 	// ULandscapePatchComponent
 	virtual bool IsAffectingWeightmapLayer(const FName& InLayerName) const override;
+	virtual bool IsEnabled() const override;
 
 	// UObject
 	virtual void PostEditChangeProperty(FPropertyChangedEvent& PropertyChangedEvent) override;
