LandscapePatch: Fix static analyzer warning about potential nullptr.

diff --git a/Engine/Plugins/Experimental/LandscapePatch/Source/LandscapePatch/Private/LandscapePatchComponent.cpp b/Engine/Plugins/Experimental/LandscapePatch/Source/LandscapePatch/Private/LandscapePatchComponent.cpp
--- a/Engine/Plugins/Experimental/LandscapePatch/Source/LandscapePatch/Private/LandscapePatchComponent.cpp
+++ b/Engine/Plugins/Experimental/LandscapePatch/Source/LandscapePatch/Private/LandscapePatchComponent.cpp
@@ -262,16 +262,19 @@ void ULandscapePatchComponent::PostEditChangeProperty(FPropertyChangedEvent& Pro
 // instance of the patch is placed in the same order in the patch  manager as the destroyed instance.
 void ULandscapePatchComponent::ApplyComponentInstanceData(FLandscapePatchComponentInstanceData* ComponentInstanceData)
 {
-	if (ComponentInstanceData && ComponentInstanceData->IndexInManager >= 0 && PatchManager.IsValid())
+	if (ComponentInstanceData)
 	{
-		PatchManager->MovePatchToIndex(this, ComponentInstanceData->IndexInManager);
-	}
+		if (ComponentInstanceData->IndexInManager >= 0 && PatchManager.IsValid())
+		{
+			PatchManager->MovePatchToIndex(this, ComponentInstanceData->IndexInManager);
+		}
 
 #if WITH_EDITOR
-	bGaveMissingPatchManagerWarning = ComponentInstanceData->bGaveMissingPatchManagerWarning;
-	bGaveNotInPatchManagerWarning = ComponentInstanceData->bGaveNotInPatchManagerWarning;
-	bGaveMissingLandscapeWarning = ComponentInstanceData->bGaveMissingLandscapeWarning;
+		bGaveMissingPatchManagerWarning = ComponentInstanceData->bGaveMissingPatchManagerWarning;
+		bGaveNotInPatchManagerWarning = ComponentInstanceData->bGaveNotInPatchManagerWarning;
+		bGaveMissingLandscapeWarning = ComponentInstanceData->bGaveMissingLandscapeWarning;
 #endif
+	}
 }
 
 void ULandscapePatchComponent::SetLandscape(ALandscape* NewLandscape)
