LandscapePatch: Prevent landscape update requests while cooking and such. Add extra information to warning messages to make the problematic patches easier to find.

diff --git a/Engine/Plugins/Experimental/LandscapePatch/Source/LandscapePatch/Private/LandscapePatchComponent.cpp b/Engine/Plugins/Experimental/LandscapePatch/Source/LandscapePatch/Private/LandscapePatchComponent.cpp
--- a/Engine/Plugins/Experimental/LandscapePatch/Source/LandscapePatch/Private/LandscapePatchComponent.cpp
+++ b/Engine/Plugins/Experimental/LandscapePatch/Source/LandscapePatch/Private/LandscapePatchComponent.cpp
@@ -385,12 +385,33 @@ void ULandscapePatchComponent::MoveToTop()
 void ULandscapePatchComponent::RequestLandscapeUpdate()
 {
 #if WITH_EDITOR
+	UWorld* World = GetWorld();
+	// Note that aside from the usual guard against doing things in the blueprint editor, the check of WorldType
+	// here also prevents us from doing the request while cooking, where WorldType is set to Inactive. Otherwise
+	// we would issue a warning below while cooking when PatchManager is not saved (and not reset in the construction
+	// script).
+	if (IsTemplate() || World->WorldType != EWorldType::Editor)
+	{
+		return;
+	}
+
+	// Helper for getting name strings in the warning messages below.
+	auto GetNameString = [](UObject* PotentiallyNullObject) {
+		FString ToReturn;
+		if (PotentiallyNullObject)
+		{
+			PotentiallyNullObject->GetName(ToReturn);
+		}
+		return ToReturn;
+	};
+
 	if (!PatchManager.IsValid())
 	{
 		if (!bGaveMissingPatchManagerWarning)
 		{
 			UE_LOG(LogLandscapePatch, Warning, TEXT("Patch does not have a valid patch manager. "
-				"Set the landscape or patch manager on the patch."));
+				"Set the landscape or patch manager on the patch. (Package: %s, Actor: %s"), 
+				*GetNameString(GetPackage()), *GetNameString(GetAttachmentRootActor()));
 			bGaveMissingPatchManagerWarning = true;
 		}
 		return;
@@ -402,7 +423,8 @@ void ULandscapePatchComponent::RequestLandscapeUpdate()
 		if (!bGaveNotInPatchManagerWarning)
 		{
 			UE_LOG(LogLandscapePatch, Warning, TEXT("Patch's patch manager does not contain this patch. "
-				"Perhaps the manager was not saved? Reset the patch manager on the patch."));
+				"Perhaps the manager was not saved? Reset the patch manager on the patch. (Package: %s, Actor: %s)"), 
+				*GetNameString(GetPackage()), *GetNameString(GetAttachmentRootActor()));
 			bGaveNotInPatchManagerWarning = true;
 		}
 		bRequestUpdate = false;
@@ -412,7 +434,8 @@ void ULandscapePatchComponent::RequestLandscapeUpdate()
 		if (!bGaveMissingLandscapeWarning)
 		{
 			UE_LOG(LogLandscapePatch, Warning, TEXT("Patch's patch manager does not have a valid owning "
-				"landscape. Perhaps the landscape was not saved? Reset the landscape on the manager."));
+				"landscape. Perhaps the landscape was not saved? Reset the landscape on the manager. (Package: %s, Manager: %s)"), 
+				*GetNameString(GetPackage()), *GetNameString(PatchManager.Get()));
 			bGaveMissingLandscapeWarning = true;
 		}
 		bRequestUpdate = false;
