LandscapePatch: Make weight patch info owning patch pointer be a weak pointer instead of raw.

diff --git a/Engine/Plugins/Experimental/LandscapePatch/Source/LandscapePatch/Private/LandscapeTexturePatch.cpp b/Engine/Plugins/Experimental/LandscapePatch/Source/LandscapePatch/Private/LandscapeTexturePatch.cpp
--- a/Engine/Plugins/Experimental/LandscapePatch/Source/LandscapePatch/Private/LandscapeTexturePatch.cpp
+++ b/Engine/Plugins/Experimental/LandscapePatch/Source/LandscapePatch/Private/LandscapeTexturePatch.cpp
@@ -1108,7 +1108,7 @@ void ULandscapeWeightPatchTextureInfo::SetSourceMode(ELandscapeTexturePatchSourc
 	Modify();
 
 	FVector2D Resolution(1, 1);
-	if (OwningPatch)
+	if (OwningPatch.IsValid())
 	{
 		Resolution = OwningPatch->GetResolution();
 	}
diff --git a/Engine/Plugins/Experimental/LandscapePatch/Source/LandscapePatch/Public/LandscapeTexturePatch.h b/Engine/Plugins/Experimental/LandscapePatch/Source/LandscapePatch/Public/LandscapeTexturePatch.h
--- a/Engine/Plugins/Experimental/LandscapePatch/Source/LandscapePatch/Public/LandscapeTexturePatch.h
+++ b/Engine/Plugins/Experimental/LandscapePatch/Source/LandscapePatch/Public/LandscapeTexturePatch.h
@@ -196,7 +196,7 @@ protected:
 
 	// Needed mainly so that we can get at the resolution...
 	UPROPERTY()
-	TObjectPtr<ULandscapeTexturePatch> OwningPatch = nullptr;
+	TWeakObjectPtr<ULandscapeTexturePatch> OwningPatch = nullptr;
 
 	// TODO: Like the similar flag for the height patch, this might not work once local merge works
 	// with landscape brushes...
