Fixed uninitialized USTRUCT param warnings

#rnx

diff --git a/Engine/Plugins/Experimental/LandscapePatch/Source/LandscapePatch/Public/LandscapeTexturePatchPS.h b/Engine/Plugins/Experimental/LandscapePatch/Source/LandscapePatch/Public/LandscapeTexturePatchPS.h
--- a/Engine/Plugins/Experimental/LandscapePatch/Source/LandscapePatch/Public/LandscapeTexturePatchPS.h
+++ b/Engine/Plugins/Experimental/LandscapePatch/Source/LandscapePatch/Public/LandscapeTexturePatchPS.h
@@ -16,13 +16,13 @@ struct FLandscapeHeightPatchConvertToNativeParams
 	GENERATED_BODY()
 
 	UPROPERTY()
-	float ZeroInEncoding;
+	float ZeroInEncoding = 0.0f;
 
 	UPROPERTY()
-	float HeightScale;
+	float HeightScale = 1.0f;
 
 	UPROPERTY()
-	float HeightOffset;
+	float HeightOffset = 0.0f;
 };
 
 namespace UE::Landscape
