LandscapePatch: Fix minor bug where patch manager is hard to reattach through detail panel to the same landscape if it was removed from that landscape in landscape mode.

diff --git a/Engine/Plugins/Experimental/LandscapePatch/Source/LandscapePatch/Private/LandscapePatchManager.cpp b/Engine/Plugins/Experimental/LandscapePatch/Source/LandscapePatch/Private/LandscapePatchManager.cpp
--- a/Engine/Plugins/Experimental/LandscapePatch/Source/LandscapePatch/Private/LandscapePatchManager.cpp
+++ b/Engine/Plugins/Experimental/LandscapePatch/Source/LandscapePatch/Private/LandscapePatchManager.cpp
@@ -142,7 +142,23 @@ void ALandscapePatchManager::SetTargetLandscape(ALandscape* InTargetLandscape)
 			OwningLandscape->RemoveBrush(this);
 		}
 
-		if (InTargetLandscape && ensure(InTargetLandscape->CanHaveLayersContent()))
+		if (!InTargetLandscape)
+		{
+			if (OwningLandscape != nullptr)
+			{
+				// This can occur if the RemoveBrush call above did not do anything because the manager
+				// was removed from the landscape in some other way (probably in landscape mode panel)
+				SetOwningLandscape(nullptr);
+			}
+			return;
+		}
+
+		if (!InTargetLandscape->CanHaveLayersContent())
+		{
+			UE_LOG(LogLandscapePatch, Warning, TEXT("Landscape target for patch manager did not have edit layers enabled. Unable to attach manager."));
+			SetOwningLandscape(nullptr);
+		}
+		else
 		{
 			static const FName PatchLayerName = FName("LandscapePatches");
 
