Adds setter for weight patches visibilitylayer modification

diff --git a/Engine/Plugins/Experimental/LandscapePatch/Source/LandscapePatch/Private/LandscapeTexturePatch.cpp b/Engine/Plugins/Experimental/LandscapePatch/Source/LandscapePatch/Private/LandscapeTexturePatch.cpp
--- a/Engine/Plugins/Experimental/LandscapePatch/Source/LandscapePatch/Private/LandscapeTexturePatch.cpp
+++ b/Engine/Plugins/Experimental/LandscapePatch/Source/LandscapePatch/Private/LandscapeTexturePatch.cpp
@@ -1392,3 +1392,13 @@ void ULandscapeTexturePatch::ClearWeightPatchBlendModeOverride(const FName& InWe
 	}
 }
 
+void ULandscapeTexturePatch::SetEditVisibilityLayer(const FName& InWeightmapLayerName, const bool bEditVisibilityLayer)
+{
+	for (TObjectPtr<ULandscapeWeightPatchTextureInfo> WeightPatch : WeightPatches)
+	{
+		if (WeightPatch->WeightmapLayerName == InWeightmapLayerName)
+		{
+			WeightPatch->bEditVisibilityLayer = bEditVisibilityLayer;
+		}
+	}
+}
diff --git a/Engine/Plugins/Experimental/LandscapePatch/Source/LandscapePatch/Public/LandscapeTexturePatch.h b/Engine/Plugins/Experimental/LandscapePatch/Source/LandscapePatch/Public/LandscapeTexturePatch.h
--- a/Engine/Plugins/Experimental/LandscapePatch/Source/LandscapePatch/Public/LandscapeTexturePatch.h
+++ b/Engine/Plugins/Experimental/LandscapePatch/Source/LandscapePatch/Public/LandscapeTexturePatch.h
@@ -420,6 +420,9 @@ public:
 	UFUNCTION(BlueprintCallable, Category = LandscapePatch)
 	virtual void ClearWeightPatchBlendModeOverride(const FName& InWeightmapLayerName);
 
+	UFUNCTION(BlueprintCallable, Category = LandscapePatch)
+	virtual void SetEditVisibilityLayer(const FName& InWeightmapLayerName, const bool bEditVisibilityLayer);
+
 
 	// These need to be public so that we can take the internal textures and write them to external ones,
 	// but unclear whether we want to expose them to blueprints, since this is a fairly internal thing.
