LandscapePatch: Fix the patch marking actors as dirty when it is hidden and unhidden.

diff --git a/Engine/Plugins/Experimental/LandscapePatch/Source/LandscapePatch/Private/LandscapePatchComponent.cpp b/Engine/Plugins/Experimental/LandscapePatch/Source/LandscapePatch/Private/LandscapePatchComponent.cpp
--- a/Engine/Plugins/Experimental/LandscapePatch/Source/LandscapePatch/Private/LandscapePatchComponent.cpp
+++ b/Engine/Plugins/Experimental/LandscapePatch/Source/LandscapePatch/Private/LandscapePatchComponent.cpp
@@ -77,6 +77,7 @@ void ULandscapePatchComponent::SetIsEnabled(bool bEnabledIn)
 		return;
 	}
 
+	Modify();
 	bIsEnabled = bEnabledIn;
 	RequestLandscapeUpdate();
 }
diff --git a/Engine/Plugins/Experimental/LandscapePatch/Source/LandscapePatch/Private/LandscapeTextureBackedRenderTarget.cpp b/Engine/Plugins/Experimental/LandscapePatch/Source/LandscapePatch/Private/LandscapeTextureBackedRenderTarget.cpp
--- a/Engine/Plugins/Experimental/LandscapePatch/Source/LandscapePatch/Private/LandscapeTextureBackedRenderTarget.cpp
+++ b/Engine/Plugins/Experimental/LandscapePatch/Source/LandscapePatch/Private/LandscapeTextureBackedRenderTarget.cpp
@@ -262,7 +262,10 @@ void ULandscapeWeightTextureBackedRenderTarget::CopyToInternalTexture()
 	}
 	else
 	{
-		InternalTexture->Modify();
+		// CopyToInternalTexture currently gets called in many non-dirty cases because we do not yet have a way to
+		// detect a true change to the render target. So, we set bAlwaysMarkDirty here to false to avoid spuriously
+		// marking the package dirty, since the internal texture may not be changing.
+		InternalTexture->Modify(false);
 	}
 
 	RenderTarget->UpdateTexture2D(InternalTexture, GetInternalTextureFormat());
@@ -369,7 +372,10 @@ void ULandscapeHeightTextureBackedRenderTarget::CopyToInternalTexture()
 	}
 	else
 	{
-		InternalTexture->Modify();
+		// CopyToInternalTexture currently gets called in many non-dirty cases because we do not yet have a way to
+		// detect a true change to the render target. So, we set bAlwaysMarkDirty here to false to avoid spuriously
+		// marking the package dirty, since the internal texture may not be changing.
+		InternalTexture->Modify(false);
 	}
 
 	UTextureRenderTarget2D* NativeEncodingRenderTarget = RenderTarget;
diff --git a/Engine/Plugins/Experimental/LandscapePatch/Source/LandscapePatch/Private/LandscapeTexturePatch.cpp b/Engine/Plugins/Experimental/LandscapePatch/Source/LandscapePatch/Private/LandscapeTexturePatch.cpp
--- a/Engine/Plugins/Experimental/LandscapePatch/Source/LandscapePatch/Private/LandscapeTexturePatch.cpp
+++ b/Engine/Plugins/Experimental/LandscapePatch/Source/LandscapePatch/Private/LandscapeTexturePatch.cpp
@@ -142,7 +142,7 @@ UTextureRenderTarget2D* ULandscapeTexturePatch::ApplyToHeightmap(UTextureRenderT
 		PatchUObject = GetHeightInternalTexture();
 		break;
 	case ELandscapeTexturePatchSourceMode::TextureBackedRenderTarget:
-		PatchUObject = GetHeightRenderTarget();
+		PatchUObject = GetHeightRenderTarget(/*bMarkDirty = */ false);
 		break;
 	case ELandscapeTexturePatchSourceMode::TextureAsset:
 
@@ -751,12 +751,12 @@ void ULandscapeTexturePatch::ReinitializeWeightPatch(ULandscapeWeightPatchTextur
 
 	FMatrix44f PatchToSource = GetPatchToHeightmapUVs(RenderTarget->SizeX, RenderTarget->SizeY, InCombinedResult->SizeX, InCombinedResult->SizeY);
 
-	ENQUEUE_RENDER_COMMAND(LandscapeTexturePatchReinitializeHeight)(
+	ENQUEUE_RENDER_COMMAND(LandscapeTexturePatchReinitializeWeight)(
 		[Source = InCombinedResult->GetResource(), Destination = RenderTarget->GetResource(), &PatchToSource](FRHICommandListImmediate& RHICmdList)
 	{
 		using namespace UE::Landscape;
 
-		FRDGBuilder GraphBuilder(RHICmdList, RDG_EVENT_NAME("LandscapeTexturePatchReinitializeHeight"));
+		FRDGBuilder GraphBuilder(RHICmdList, RDG_EVENT_NAME("LandscapeTexturePatchReinitializeWeight"));
 
 		FReinitializeLandscapePatchPS::FParameters* ShaderParams = GraphBuilder.AllocParameters<FReinitializeLandscapePatchPS::FParameters>();
 
@@ -1160,6 +1160,15 @@ void ULandscapeTexturePatch::SetHeightTextureAsset(UTexture* TextureIn)
 	HeightTextureAsset = TextureIn;
 }
 
+UTextureRenderTarget2D* ULandscapeTexturePatch::GetHeightRenderTarget(bool bMarkDirty)
+{
+	if (bMarkDirty)
+	{
+		MarkPackageDirty();
+	}
+	return HeightInternalData ? HeightInternalData->GetRenderTarget() : nullptr;
+}
+
 void ULandscapeTexturePatch::ResetHeightEncodingMode(ELandscapeTextureHeightPatchEncoding EncodingMode)
 {
 	Modify();
diff --git a/Engine/Plugins/Experimental/LandscapePatch/Source/LandscapePatch/Public/LandscapeTexturePatch.h b/Engine/Plugins/Experimental/LandscapePatch/Source/LandscapePatch/Public/LandscapeTexturePatch.h
--- a/Engine/Plugins/Experimental/LandscapePatch/Source/LandscapePatch/Public/LandscapeTexturePatch.h
+++ b/Engine/Plugins/Experimental/LandscapePatch/Source/LandscapePatch/Public/LandscapeTexturePatch.h
@@ -126,8 +126,8 @@ struct LANDSCAPEPATCH_API FLandscapeTexturePatchEncodingSettings
 	GENERATED_BODY()
 public:
 	/**
-	 * The value in the patch data that corresponds to 0 landscape height (which is in line with patch Z when
-	 * "Use Patch Z As Reference" is true, and at landscape zero/mid value when false).
+	 * The value in the patch data that corresponds to 0 height relative to the starting point
+	 * specified by Zero Height Meaning.
 	 */
 	UPROPERTY(EditAnywhere, BlueprintReadWrite, Category = Settings)
 	double ZeroInEncoding = 0;
@@ -322,11 +322,14 @@ public:
 	UFUNCTION(BlueprintCallable, Category = LandscapePatch)
 	void SetHeightTextureAsset(UTexture* TextureIn);
 
+	/**
+	 * Gets the internal height render target, if source mode is set to Texture Backed Render Target.
+	 * 
+	 * @param bMarkDirty If true, marks the containing package as dirty, since the render target is presumably
+	 *  being written to. Can be set to false if the render target is not being written to.
+	 */
 	UFUNCTION(BlueprintCallable, Category = LandscapePatch)
-	virtual UTextureRenderTarget2D* GetHeightRenderTarget() 
-	{ 
-		return HeightInternalData ? HeightInternalData->GetRenderTarget() : nullptr; 
-	}
+	virtual UTextureRenderTarget2D* GetHeightRenderTarget(bool bMarkDirty = true);
 
 	UFUNCTION(BlueprintCallable, Category = LandscapePatch, meta = (ETextureRenderTargetFormat = "ETextureRenderTargetFormat::RTF_R32f"))
 	void SetHeightRenderTargetFormat(ETextureRenderTargetFormat Format);
@@ -494,7 +497,7 @@ protected:
 	/** How the values stored in the patch represent the height. Not customizable for Internal Texture source mode, which always uses native packed height. */
 	UPROPERTY(EditAnywhere, Category = HeightPatch, meta = (
 		EditCondition = "HeightSourceMode != ELandscapeTexturePatchSourceMode::InternalTexture"))
-	ELandscapeTextureHeightPatchEncoding HeightEncoding = ELandscapeTextureHeightPatchEncoding::ZeroToOne;
+	ELandscapeTextureHeightPatchEncoding HeightEncoding = ELandscapeTextureHeightPatchEncoding::WorldUnits;
 
 	/** Encoding settings. Not relevant when using native packed height as the encoding. */
 	UPROPERTY(EditAnywhere, Category = HeightPatch, meta = (UIMin = "0", UIMax = "1",
