LandscapeTexturePatch: Make the texture backed render targets use a hash of their texture for its DDC guid to avoid spurious recompiles if it is reset in a construction script.

diff --git a/Engine/Plugins/Experimental/LandscapePatch/Source/LandscapePatch/Private/LandscapeTextureBackedRenderTarget.cpp b/Engine/Plugins/Experimental/LandscapePatch/Source/LandscapePatch/Private/LandscapeTextureBackedRenderTarget.cpp
--- a/Engine/Plugins/Experimental/LandscapePatch/Source/LandscapePatch/Private/LandscapeTextureBackedRenderTarget.cpp
+++ b/Engine/Plugins/Experimental/LandscapePatch/Source/LandscapePatch/Private/LandscapeTextureBackedRenderTarget.cpp
@@ -25,6 +25,9 @@ namespace LandscapeTextureBackedRenderTargetLocals
 		// TODO: How do we allow lossless compression, instead of disallowing compression entirely? Just setting 
 		// LossyCompressionAmount to ETextureLossyCompressionAmount::TLCA_None is not sufficient.
 		Texture->CompressionNone = true;
+		// If we don't do this, then each newly created empty texture created will request the DDC to re-compile it,
+		// which is problematic if the texture is being recreated over and over in a construction script.
+		Texture->Source.UseHashAsGuid();
 
 		return Texture;
 	}
