LandscapePatch: Fix "Reinitialize Height" button now working with some height encoding settings. Add some helpful comments

diff --git a/Engine/Plugins/Experimental/LandscapePatch/Source/LandscapePatch/Private/LandscapeTexturePatch.cpp b/Engine/Plugins/Experimental/LandscapePatch/Source/LandscapePatch/Private/LandscapeTexturePatch.cpp
--- a/Engine/Plugins/Experimental/LandscapePatch/Source/LandscapePatch/Private/LandscapeTexturePatch.cpp
+++ b/Engine/Plugins/Experimental/LandscapePatch/Source/LandscapePatch/Private/LandscapeTexturePatch.cpp
@@ -472,7 +472,19 @@ void ULandscapeTexturePatch::GetWeightShaderParams(const FIntPoint& SourceResolu
 	ParamsOut.InFlags = static_cast<uint8>(Flags);
 }
 
-FLandscapeHeightPatchConvertToNativeParams ULandscapeTexturePatch::GetHeightConversionParams() const
+// This function determines how our internal height render targets get converted to the format that gets
+// serialized. In a perfect world, this largely shouldn't matter as long as we don't lose data in the conversion
+// back and forth. In practice, it matters for transitioning the SourceMode between ELandscapeTexturePatchSourceMode::InternalTexture 
+// and ELandscapeTexturePatchSourceMode::TextureBackedRenderTarget, and it matters for reinitializing the patch
+// from the current landscape. In the former, it matters because the transition is easy if the backing format
+// is the same as the equivalent texture. In the latter, it matters because the reinitialization is easy if
+// the backing format is the same as the applied landscape values. Currently we end up making the former easy, i.e.
+// we serialize render targets to their equivalent native texture representation, and don't bake in the offset.
+// This means that we need to do a bit more work when reinitializing to account for the offset.
+// It should also be noted that there are some truncation/rounding implications to the choices made here that
+// only matter if the user is messing around with the conversion parameters and hoping not to lose data... But
+// there's a limited amount that we can protect the user in that case anyway.
+FLandscapeHeightPatchConvertToNativeParams ULandscapeTexturePatch::GetHeightConvertToNativeParams() const
 {
 	// When doing conversions, we bake into a height in the same way that we do when applying the patch.
 
@@ -483,12 +495,7 @@ FLandscapeHeightPatchConvertToNativeParams ULandscapeTexturePatch::GetHeightConv
 	LandscapeHeightScale = LandscapeHeightScale == 0 ? 1 : LandscapeHeightScale;
 	ConversionParams.HeightScale = HeightEncodingSettings.WorldSpaceEncodingScale * LANDSCAPE_INV_ZSCALE / LandscapeHeightScale;
 
-	// TODO: We can choose whether we want to bake in the height offset if it exists. Doing so will handle
-	// some edge cases where the value stored in the patch is outside the range storeable in the native format
-	// normally, but within the range of the landscape due to the patch being far above/below the landscape to
-	// compensate. However, while this is good for conversions for the purposes of serialization, it's not good
-	// for conversions for the purposes of source mode change, so we would need to do things slightly differently
-	// in the two cases. For now, we'll just not bother with that (unlikely?) edge case.
+	// See above discussion about why we don't currently bake in height offset.
 	ConversionParams.HeightOffset = 0;
 
 	return ConversionParams;
@@ -620,9 +627,12 @@ void ULandscapeTexturePatch::ReinitializeHeight(UTextureRenderTarget2D* InCombin
 
 	// The way we're going to do it is that we'll copy the packed values directly to a temporary render target, offset 
 	// them if needed (to undo whatever offsetting will happen during application), and store the result directly in the
-	// internal texture. Then we'll update the actual associated render target from the internal texture (if needed) so
+	// backing internal texture. Then we'll update the actual associated render target from the internal texture (if needed) so
 	// that unpacking and height format conversion happens the same way as everywhere else.
 
+	// We do need to make sure that the scale conversion for the backing texture matches what will be used when applying it.
+	UpdateHeightConvertToNativeParamsIfNeeded();
+
 	UTextureRenderTarget2D* TemporaryNativeHeightCopy = NewObject<UTextureRenderTarget2D>(this);
 	TemporaryNativeHeightCopy->RenderTargetFormat = ETextureRenderTargetFormat::RTF_RGBA8;
 	TemporaryNativeHeightCopy->InitAutoFormat(ResolutionX, ResolutionY);
@@ -1081,6 +1091,11 @@ void ULandscapeTexturePatch::PostEditChangeProperty(FPropertyChangedEvent& Prope
 				NumWeightPatches = WeightPatches.Num();
 			}
 		}
+		else if (PropertyChangedEvent.Property->GetFName() == GET_MEMBER_NAME_CHECKED(FLandscapeTexturePatchEncodingSettings, ZeroInEncoding)
+			|| PropertyChangedEvent.Property->GetFName() == GET_MEMBER_NAME_CHECKED(FLandscapeTexturePatchEncodingSettings, WorldSpaceEncodingScale))
+		{
+			UpdateHeightConvertToNativeParamsIfNeeded();
+		}
 	}
 
 	Super::PostEditChangeProperty(PropertyChangedEvent);
@@ -1162,7 +1177,7 @@ void ULandscapeTexturePatch::SetHeightSourceMode(ELandscapeTexturePatchSourceMod
 		ULandscapeHeightTextureBackedRenderTarget* InternalDataToReturn = NewObject<ULandscapeHeightTextureBackedRenderTarget>(this);
 		InternalDataToReturn->SetResolution(ResolutionX, ResolutionY);
 		InternalDataToReturn->SetFormat(HeightRenderTargetFormat);
-		InternalDataToReturn->ConversionParams = GetHeightConversionParams();
+		InternalDataToReturn->ConversionParams = GetHeightConvertToNativeParams();
 
 		return InternalDataToReturn;
 	});
@@ -1188,6 +1203,31 @@ UTextureRenderTarget2D* ULandscapeTexturePatch::GetHeightRenderTarget(bool bMark
 	return HeightInternalData ? HeightInternalData->GetRenderTarget() : nullptr;
 }
 
+void ULandscapeTexturePatch::UpdateHeightConvertToNativeParamsIfNeeded()
+{
+#if WITH_EDITOR
+	if (HeightInternalData && Landscape.IsValid() && PatchManager.IsValid())
+	{
+		FLandscapeHeightPatchConvertToNativeParams ConversionParams = GetHeightConvertToNativeParams();
+		if (ConversionParams.HeightScale == 0)
+		{
+			// If the scale is 0, then storing in the texture would lose the data we have,
+			// so keep whatever the previous storage encoding was if nonzero, otherwise set to 1.
+			ConversionParams.HeightScale = HeightInternalData->ConversionParams.HeightScale != 0 ? HeightInternalData->ConversionParams.HeightScale
+				: 1;
+		}
+		
+		if (ConversionParams.ZeroInEncoding != HeightInternalData->ConversionParams.ZeroInEncoding
+			|| ConversionParams.HeightScale != HeightInternalData->ConversionParams.HeightScale
+			|| ConversionParams.HeightOffset != HeightInternalData->ConversionParams.HeightOffset)
+		{
+			HeightInternalData->Modify();
+			HeightInternalData->ConversionParams = ConversionParams;
+		}
+	}
+#endif
+}
+
 void ULandscapeTexturePatch::ResetHeightEncodingMode(ELandscapeTextureHeightPatchEncoding EncodingMode)
 {
 	Modify();
@@ -1202,6 +1242,18 @@ void ULandscapeTexturePatch::ResetHeightEncodingMode(ELandscapeTextureHeightPatc
 		HeightEncodingSettings.ZeroInEncoding = 0;
 		HeightEncodingSettings.WorldSpaceEncodingScale = 1;
 	}
+	SetHeightRenderTargetFormat(EncodingMode == ELandscapeTextureHeightPatchEncoding::NativePackedHeight ?
+		ETextureRenderTargetFormat::RTF_RGBA8 : ETextureRenderTargetFormat::RTF_R32f);
+
+	UpdateHeightConvertToNativeParamsIfNeeded();
+}
+
+void ULandscapeTexturePatch::SetHeightEncodingSettings(const FLandscapeTexturePatchEncodingSettings& Settings)
+{
+	Modify();
+	HeightEncodingSettings = Settings;
+
+	UpdateHeightConvertToNativeParamsIfNeeded();
 }
 
 void ULandscapeTexturePatch::SetHeightRenderTargetFormat(ETextureRenderTargetFormat Format)
diff --git a/Engine/Plugins/Experimental/LandscapePatch/Source/LandscapePatch/Public/LandscapeTexturePatch.h b/Engine/Plugins/Experimental/LandscapePatch/Source/LandscapePatch/Public/LandscapeTexturePatch.h
--- a/Engine/Plugins/Experimental/LandscapePatch/Source/LandscapePatch/Public/LandscapeTexturePatch.h
+++ b/Engine/Plugins/Experimental/LandscapePatch/Source/LandscapePatch/Public/LandscapeTexturePatch.h
@@ -298,6 +298,9 @@ public:
 		Falloff = FalloffIn; 
 	}
 
+	/**
+	 * Determines how the height patch is blended into the existing terrain.
+	 */
 	UFUNCTION(BlueprintCallable, Category = LandscapePatch)
 	void SetBlendMode(ELandscapeTexturePatchBlendMode BlendModeIn) 
 	{ 
@@ -328,6 +331,19 @@ public:
 	/**
 	 * Gets the internal height render target, if source mode is set to Texture Backed Render Target.
 	 * 
+	 * Things that should be set up if using the internal render target:
+	 * - SetHeightSourceMode should have been called with TextureBackedRenderTarget.
+	 * - An appropriate texture size should have been set with SetResolution. If the patch extent has already
+	 *  been set, you can base your resolution on the extent and the resolution of the landscape by using
+	 *  GetInitResolutionFromLandscape().
+	 * - SetHeightRenderTargetFormat should have been called with a desired format. In particular, if using
+	 *  an alpha channel, the format should have an alpha channel (and SetUseAlphaChannelForHeight should have
+	 *  been called with "true").
+	 * 
+	 * In addition, you may need to call SetHeightEncodingMode, SetHeightEncodingSettings, and SetZeroHeightMeaning
+	 * based on how you want the data you write to be interpreted. This part is not specific to using an internal render
+	 * target, since you are likely to need to do that with a TextureAsset source mode as well.
+	 * 
 	 * @param bMarkDirty If true, marks the containing package as dirty, since the render target is presumably
 	 *  being written to. Can be set to false if the render target is not being written to.
 	 */
@@ -337,6 +353,11 @@ public:
 	UFUNCTION(BlueprintCallable, Category = LandscapePatch, meta = (ETextureRenderTargetFormat = "ETextureRenderTargetFormat::RTF_R32f"))
 	void SetHeightRenderTargetFormat(ETextureRenderTargetFormat Format);
 
+	/**
+	 * Determines whether the height patch alpha channel is used for blending into the existing values.
+	 * Note that the source data needs to have an alpha channel in this case. How the alpha channel is
+	 * used depends on the patch blend mode (see SetBlendMode).
+	 */
 	UFUNCTION(BlueprintCallable, Category = "LandscapePatch")
 	void SetUseAlphaChannelForHeight(bool bUse)
 	{ 
@@ -356,7 +377,8 @@ public:
 	}
 
 	/**
-	 * Just like SetSourceEncodingMode, but resets ZeroInEncoding and WorldSpaceEncodingScale to mode-specific defaults.
+	 * Just like SetSourceEncodingMode, but resets ZeroInEncoding, WorldSpaceEncodingScale, and height
+	 * render target format to mode-specific defaults.
 	 */
 	UFUNCTION(BlueprintCallable, Category = LandscapePatch)
 	void ResetHeightEncodingMode(ELandscapeTextureHeightPatchEncoding EncodingMode);
@@ -367,11 +389,7 @@ public:
 	 * in the same space as the landscape heightmap.
 	 */
 	UFUNCTION(BlueprintCallable, Category = LandscapePatch)
-	void SetHeightEncodingSettings(const FLandscapeTexturePatchEncodingSettings& Settings)
-	{
-		Modify();
-		HeightEncodingSettings = Settings;
-	}
+	void SetHeightEncodingSettings(const FLandscapeTexturePatchEncodingSettings& Settings);
 
 	/**
 	 * Set how zero height is interpreted, see comments in ELandscapeTextureHeightPatchZeroHeightMeaning.
@@ -492,7 +510,10 @@ protected:
 	TObjectPtr<UTexture> HeightTextureAsset = nullptr;
 
 	
-	/** When true, texture alpha channel will be used when applying the patch. */
+	/** 
+	 * When true, texture alpha channel will be used when applying the patch. Note that the source data needs to
+	 * have an alpha channel for this to have an effect.
+	 */
 	UPROPERTY(EditAnywhere, Category = HeightPatch)
 	bool bUseTextureAlphaForHeight = false;
 
@@ -585,8 +606,9 @@ protected:
 	int32 InitTextureSizeY = 33;
 
 private:
+	void UpdateHeightConvertToNativeParamsIfNeeded();
 #if WITH_EDITOR
-	FLandscapeHeightPatchConvertToNativeParams GetHeightConversionParams() const;
+	FLandscapeHeightPatchConvertToNativeParams GetHeightConvertToNativeParams() const;
 	UTextureRenderTarget2D* ApplyToHeightmap(UTextureRenderTarget2D* InCombinedResult);
 	UTextureRenderTarget2D* ApplyToWeightmap(ULandscapeWeightPatchTextureInfo* PatchInfo, UTextureRenderTarget2D* InCombinedResult);
 
