Landscape: Fix landscape patch breakage due to a typo in 28494111

diff --git a/Engine/Plugins/Experimental/LandscapePatch/Source/LandscapePatch/Private/LandscapeTextureBackedRenderTarget.cpp b/Engine/Plugins/Experimental/LandscapePatch/Source/LandscapePatch/Private/LandscapeTextureBackedRenderTarget.cpp
--- a/Engine/Plugins/Experimental/LandscapePatch/Source/LandscapePatch/Private/LandscapeTextureBackedRenderTarget.cpp
+++ b/Engine/Plugins/Experimental/LandscapePatch/Source/LandscapePatch/Private/LandscapeTextureBackedRenderTarget.cpp
@@ -424,7 +424,7 @@ void ULandscapeHeightTextureBackedRenderTarget::CopyToInternalTexture()
 
 	// This call does a flush for us, so the render target should be updated.
 	FText ErrorMessage;
-	if (RenderTarget->UpdateTexture(InternalTexture, CTF_Default, /*InAlphaOverride = */nullptr, /*InTextureChangingDelegate =*/ [](UTexture*) {}, &ErrorMessage))
+	if (NativeEncodingRenderTarget->UpdateTexture(InternalTexture, CTF_Default, /*InAlphaOverride = */nullptr, /*InTextureChangingDelegate =*/ [](UTexture*) {}, &ErrorMessage))
 	{
 		check(InternalTexture->Source.GetFormat() == GetInternalTextureFormat());
 		InternalTexture->UpdateResource();
