Fix minor shader warnings in plugin

diff --git a/Engine/Plugins/Experimental/LandscapePatch/Shaders/Private/LandscapeTexturePatchPS.usf b/Engine/Plugins/Experimental/LandscapePatch/Shaders/Private/LandscapeTexturePatchPS.usf
--- a/Engine/Plugins/Experimental/LandscapePatch/Shaders/Private/LandscapeTexturePatchPS.usf
+++ b/Engine/Plugins/Experimental/LandscapePatch/Shaders/Private/LandscapeTexturePatchPS.usf
@@ -223,7 +223,7 @@ void ApplyLandscapeTextureWeightPatch(in float4 SVPos : SV_POSITION, out float O
 	float PatchWeight = PatchSampledValue.x;
 	
 	int2 WeightmapCoordinates = floor(SVPos.xy);
-	float4 CurrentWeight = InSourceWeightmap.Load(int3(WeightmapCoordinates, 0)).x;
+	float CurrentWeight = InSourceWeightmap.Load(int3(WeightmapCoordinates, 0)).x;
 	
 	float Alpha = GetFalloffAlpha(InFalloffWorldMargin, InPatchWorldDimensions, PatchUVCoordinates, InEdgeUVDeadBorder, bRectangularFalloff);
 	
