LanscapePatch: Expose patch manager "move patch to index" and "get index of patch" to blueprints to give more control over patch ordering.

diff --git a/Engine/Plugins/Experimental/LandscapePatch/Source/LandscapePatch/Private/LandscapePatchManager.cpp b/Engine/Plugins/Experimental/LandscapePatch/Source/LandscapePatch/Private/LandscapePatchManager.cpp
--- a/Engine/Plugins/Experimental/LandscapePatch/Source/LandscapePatch/Private/LandscapePatchManager.cpp
+++ b/Engine/Plugins/Experimental/LandscapePatch/Source/LandscapePatch/Private/LandscapePatchManager.cpp
@@ -180,17 +180,17 @@ void ALandscapePatchManager::SetTargetLandscape(ALandscape* InTargetLandscape)
 #endif
 }
 
-bool ALandscapePatchManager::ContainsPatch(TObjectPtr<ULandscapePatchComponent> Patch) const
+bool ALandscapePatchManager::ContainsPatch(ULandscapePatchComponent* Patch) const
 {
 	return PatchComponents.Contains(Patch);
 }
 
-void ALandscapePatchManager::AddPatch(TObjectPtr<ULandscapePatchComponent> Patch)
+void ALandscapePatchManager::AddPatch(ULandscapePatchComponent* Patch)
 {
 	if (Patch)
 	{
 		Modify();
-		PatchComponents.AddUnique(TSoftObjectPtr<ULandscapePatchComponent>(Patch.Get()));
+		PatchComponents.AddUnique(TSoftObjectPtr<ULandscapePatchComponent>(Patch));
 
 		// No need to update if the patch is disabled. Important to avoid needlessly updating while dragging a blueprint with
 		// a disabled patch (since construction scripts constantly add and remove).
@@ -201,14 +201,14 @@ void ALandscapePatchManager::AddPatch(TObjectPtr<ULandscapePatchComponent> Patch
 	}
 }
 
-bool ALandscapePatchManager::RemovePatch(TObjectPtr<ULandscapePatchComponent> Patch)
+bool ALandscapePatchManager::RemovePatch(ULandscapePatchComponent* Patch)
 {
 	bool bRemoved = false;
 
 	if (Patch)
 	{
 		Modify();
-		bRemoved = PatchComponents.Remove(TSoftObjectPtr<ULandscapePatchComponent>(Patch.Get())) > 0;
+		bRemoved = PatchComponents.Remove(TSoftObjectPtr<ULandscapePatchComponent>(Patch)) > 0;
 
 		// No need to update if the patch was already disabled.Important to avoid needlessly updating while dragging 
 		// a blueprint with a disabled patch (since construction scripts constantly add and remove).
@@ -221,12 +221,12 @@ bool ALandscapePatchManager::RemovePatch(TObjectPtr<ULandscapePatchComponent> Pa
 	return bRemoved;
 }
 
-int32 ALandscapePatchManager::GetIndexOfPatch(TObjectPtr<const ULandscapePatchComponent> Patch) const
+int32 ALandscapePatchManager::GetIndexOfPatch(const ULandscapePatchComponent* Patch) const
 {
 	return PatchComponents.IndexOfByKey(Patch);
 }
 
-void ALandscapePatchManager::MovePatchToIndex(TObjectPtr<ULandscapePatchComponent> Patch, int32 Index)
+void ALandscapePatchManager::MovePatchToIndex(ULandscapePatchComponent* Patch, int32 Index)
 {
 	if (!Patch || Index < 0 || GetIndexOfPatch(Patch) == Index)
 	{
@@ -239,7 +239,7 @@ void ALandscapePatchManager::MovePatchToIndex(TObjectPtr<ULandscapePatchComponen
 	RemovePatch(Patch);
 
 	Index = FMath::Clamp(Index, 0, PatchComponents.Num());
-	PatchComponents.Insert(TSoftObjectPtr<ULandscapePatchComponent>(Patch.Get()), Index);
+	PatchComponents.Insert(TSoftObjectPtr<ULandscapePatchComponent>(Patch), Index);
 
 	if (Patch->IsEnabled())
 	{
diff --git a/Engine/Plugins/Experimental/LandscapePatch/Source/LandscapePatch/Public/LandscapePatchManager.h b/Engine/Plugins/Experimental/LandscapePatch/Source/LandscapePatch/Public/LandscapePatchManager.h
--- a/Engine/Plugins/Experimental/LandscapePatch/Source/LandscapePatch/Public/LandscapePatchManager.h
+++ b/Engine/Plugins/Experimental/LandscapePatch/Source/LandscapePatch/Public/LandscapePatchManager.h
@@ -36,6 +36,7 @@ public:
 
 	// Adds the brush to the given landscape, removing it from any previous one. This differs from SetOwningLandscape
 	// in that SetOwningLandscape is called by the landscape itself from AddBrushToLayer to update the manager.
+	UFUNCTION(BlueprintCallable, Category = LandscapeManager)
 	virtual void SetTargetLandscape(ALandscape* InOwningLandscape);
 
 	// For use by the owned patch objects.
@@ -46,23 +47,28 @@ public:
 	 */
 	virtual FTransform GetHeightmapCoordsToWorld() { return HeightmapCoordsToWorld; }
 
-	bool ContainsPatch(TObjectPtr<ULandscapePatchComponent> Patch) const;
+	UFUNCTION(BlueprintCallable, Category = LandscapePatch)
+	bool ContainsPatch(ULandscapePatchComponent* Patch) const;
 
-	void AddPatch(TObjectPtr<ULandscapePatchComponent> Patch);
+	UFUNCTION(BlueprintCallable, Category = LandscapePatch)
+	void AddPatch(ULandscapePatchComponent* Patch);
 
-	bool RemovePatch(TObjectPtr<ULandscapePatchComponent> Patch);
+	UFUNCTION(BlueprintCallable, Category = LandscapePatch)
+	bool RemovePatch(ULandscapePatchComponent* Patch);
 
 	/** 
 	 * Gets the index of a particular patch in the manager's stack of patches (later indices get applied after
 	 * earlier ones.
 	 */
-	int32 GetIndexOfPatch(TObjectPtr<const ULandscapePatchComponent> Patch) const;
+	UFUNCTION(BlueprintCallable, Category = LandscapePatch)
+	int32 GetIndexOfPatch(const ULandscapePatchComponent* Patch) const;
 
 	/**
 	 * Moves patch to given index in the list of patches held by the manager (so that it is applied at 
 	 * a particular time relative to the others).
 	 */
-	void MovePatchToIndex(TObjectPtr<ULandscapePatchComponent> Patch, int32 Index);
+	UFUNCTION(BlueprintCallable, Category = LandscapePatch)
+	void MovePatchToIndex(ULandscapePatchComponent* Patch, int32 Index);
 	
 #if WITH_EDITOR
 	// ALandscapeBlueprintBrushBase
