LandscapePatch: Hardening, particularly dealing with situations where patches might be saved without a proper patch manager pointer. Fixing those situations should fix patch order not being preserved across loads. Changes:

- Add ability to detect when blueprint owned patches were saved without a patch manager and are getting one auto assigned, or patch manager is being modified in a construction script rerun (dangerous because patch order can end up being dictated by load order).
- Add cvar to disable auto-finding of patch manager in construction script. Sadly we probably can't make this the default since this might be accidentally relied upon.
- Add map checks for various mismatches of patch and patch manager saving (saving one but not the other).
- Add console commands to fix up all patch pointers, or those owned by patch managers.
- Fix crash that could happen if texture patch tried to run while not having its patch manager pointer set properly.
- Fix ensure that fires in visualizer on deleting patches.
- Fix potentially not carrying over the original patch manager across construction script reruns.

diff --git a/Engine/Plugins/Experimental/LandscapePatch/Source/LandscapePatch/Private/LandscapePatchComponent.cpp b/Engine/Plugins/Experimental/LandscapePatch/Source/LandscapePatch/Private/LandscapePatchComponent.cpp
--- a/Engine/Plugins/Experimental/LandscapePatch/Source/LandscapePatch/Private/LandscapePatchComponent.cpp
+++ b/Engine/Plugins/Experimental/LandscapePatch/Source/LandscapePatch/Private/LandscapePatchComponent.cpp
@@ -2,12 +2,19 @@
 
 #include "LandscapePatchComponent.h"
 
+#include "CoreGlobals.h" // GIsReconstructingBlueprintInstances
 #include "EngineUtils.h"
+#include "HAL/IConsoleManager.h" // FAutoConsoleVariableRef
 #include "Landscape.h"
 #include "LandscapePatchLogging.h"
 #include "LandscapePatchManager.h"
+#include "Logging/MessageLog.h"
+#include "Misc/UObjectToken.h"
 #include "PropertyPairsMap.h"
+#include "UObject/ObjectSaveContext.h"
 #include "UObject/Package.h"
+#include "UObject/UObjectBaseUtility.h" // GetNameSafe
+#include "UObject/UObjectIterator.h"
 
 #include UE_INLINE_GENERATED_CPP_BY_NAME(LandscapePatchComponent)
 
@@ -56,6 +63,70 @@ namespace LandscapePatchComponentLocals
 		return PatchManager;
 	}
 
+	bool bAllowAutoFindOfManagerInConstructionScript = true;
+	static FAutoConsoleVariableRef CVarReAddToPatchManagerOnConstructionIfLost(
+		TEXT("LandscapePatch.AllowAutoFindOfManagerInConstructionScript"),
+		bAllowAutoFindOfManagerInConstructionScript,
+		TEXT("If a blueprint-contained patch does not have a patch manager when its construction script is rerun, "
+			"allows for a patch manager to be found automatically. This is more safe to have turned off because "
+			"it can hide potential issues when patches or manager were not properly saved, and as a result, patch "
+			"order could change depending on load order of actors. It is on by default for now for legacy support."));
+
+
+	// A way to fix patch manager pointers. More aggressive than the map check, because it allows auto-finding
+	// of a new manager.
+	FAutoConsoleCommand CCmdFixPatchManagerPointers(
+		TEXT("LandscapePatch.FixPatchManagerPointers"),
+		TEXT("For all patches, make sure that patch manager pointers exist, contain this patch, and patches "
+			"and managers are dirtied if they changed. Patches without a manager get one auto-assigned."),
+		FConsoleCommandDelegate::CreateLambda([]() 
+	{
+		// Iterate through all patches
+		for (TObjectIterator<ULandscapePatchComponent> It(
+			/*AdditionalExclusionFlags = */RF_ClassDefaultObject,
+			/*bIncludeDerivedClasses = */true,
+			/*InInternalExclusionFlags = */EInternalObjectFlags::Garbage); It; ++It)
+		{
+			ULandscapePatchComponent* Patch = *It;
+			if (!IsValid(Patch))
+			{
+				continue;
+			}
+
+			UWorld* World = Patch->GetWorld();
+			if (Patch->IsTemplate() || !IsValid(World) || World->WorldType != EWorldType::Editor)
+			{
+				continue;
+			}
+
+			Patch->SetUpPatchManagerData();
+			Patch->MarkDirtyIfModifiedInConstructionScript();
+		}
+
+		for (TObjectIterator<ALandscapePatchManager> It(
+			/*AdditionalExclusionFlags = */RF_ClassDefaultObject,
+			/*bIncludeDerivedClasses = */true,
+			/*InInternalExclusionFlags = */EInternalObjectFlags::Garbage); It; ++It)
+		{
+			ALandscapePatchManager* Manager = *It;
+			if (!IsValid(Manager))
+			{
+				continue;
+			}
+
+			UWorld* World = Manager->GetWorld();
+			if (Manager->IsTemplate() || !IsValid(World) || World->WorldType != EWorldType::Editor)
+			{
+				continue;
+			}
+
+			if (Manager)
+			{
+				Manager->MarkDirtyIfModifiedInConstructionScript();
+			}
+		}
+	}));
+
 #endif
 }
 
@@ -83,6 +154,82 @@ void ULandscapePatchComponent::SetIsEnabled(bool bEnabledIn)
 }
 
 #if WITH_EDITOR
+void ULandscapePatchComponent::CheckForErrors()
+{
+	Super::CheckForErrors();
+
+	using namespace LandscapePatchComponentLocals;
+	
+	if (!IsRealPatch())
+	{
+		return;
+	}
+
+	auto GetPackageAndActorArgs = [this]()
+	{
+		FFormatNamedArguments Arguments;
+		Arguments.Add(TEXT("Package"), FText::FromString(*GetNameSafe(GetPackage())));
+		Arguments.Add(TEXT("Actor"), FText::FromString(*GetNameSafe(GetAttachmentRootActor())));
+		return Arguments;
+	};
+
+	if (PatchManager.IsValid() != Landscape.IsValid())
+	{
+		FMessageLog("MapCheck").Warning()
+			->AddToken(FUObjectToken::Create(this))
+			->AddToken(FTextToken::Create(FText::Format(LOCTEXT("PatchManagerAndLandscapeDisagree", "Patch has inconsistent "
+				"manager and landscape pointers. (Package: {Package}, Actor: {Actor}). "
+				"Fix individually or run LandscapePatch.FixPatchManagerPointers."), GetPackageAndActorArgs())))
+			->AddToken(FActionToken::Create(LOCTEXT("FixInconsistentPointersButton", "Fix inconsistent pointers"), FText(),
+				FOnActionTokenExecuted::CreateWeakLambda(this, [this]() 
+		{
+			if (PatchManager.IsValid())
+			{
+				SetPatchManager(PatchManager.Get());
+			}
+			else if (Landscape.IsValid())
+			{
+				SetLandscape(Landscape.Get());
+			}
+		})));
+	}
+
+	if (PatchManager.IsValid() && !PatchManager->ContainsPatch(this))
+	{
+		FMessageLog("MapCheck").Warning()
+			->AddToken(FUObjectToken::Create(this))
+			->AddToken(FTextToken::Create(FText::Format(LOCTEXT("PatchNotInManager", "Patch was not found in its patch manager. "
+				"(Package: {Package}, Actor: {Actor}). "
+				"Fix individually or run LandscapePatch.FixPatchManagerPointers."), GetPackageAndActorArgs())))
+			->AddToken(FActionToken::Create(LOCTEXT("AddToManagerButton", "Add to manager"), FText(),
+				FOnActionTokenExecuted::CreateWeakLambda(this, [this]() 
+		{
+			if (PatchManager.IsValid())
+			{
+				SetPatchManager(PatchManager.Get());
+			}
+		})));
+	}
+
+	if (bDirtiedByConstructionScript && !(GetPackage() && GetPackage()->IsDirty()))
+	{
+		FMessageLog("MapCheck").Warning()
+			->AddToken(FUObjectToken::Create(this))
+			->AddToken(FTextToken::Create(FText::Format(LOCTEXT("PatchNeedsSaving", "Patch got modified in a construction "
+				"script rerun but has not been marked dirty. (Package: {Package}, Actor: {Actor}). "
+				"Fix individually or run LandscapePatch.FixPatchManagerPointers."), GetPackageAndActorArgs())))
+			->AddToken(FActionToken::Create(LOCTEXT("MarkDirtyButton", "Mark dirty"), FText(),
+				FOnActionTokenExecuted::CreateWeakLambda(this, [this]() 
+		{
+			MarkPackageDirty();
+			if (PatchManager.IsValid())
+			{
+				PatchManager->MarkDirtyIfModifiedInConstructionScript();
+			}
+		})));
+	}
+}
+
 void ULandscapePatchComponent::OnComponentCreated()
 {
 	using namespace LandscapePatchComponentLocals;
@@ -93,54 +240,105 @@ void ULandscapePatchComponent::OnComponentCreated()
 	bWasCopy = bPropertiesCopiedIndicator;
 	bPropertiesCopiedIndicator = true;
 
-	UWorld* World = GetWorld();
-	if (IsTemplate() || World->WorldType != EWorldType::Editor)
+	if (!IsRealPatch())
+	{
+		return;
+	}
+
+	// Doing stuff during construction script reruns is a huge pain, and ideally we would always exit right here.
+	// However we preserve legacy behavior where patches added themselves to the manager from OnComponentCreated
+	// even in the construction script, if bAllowAutoFindOfManagerInConstructionScript is true.
+	// See ApplyComponentInstanceData to see how we try to detect that patch got modified in the construction
+	// script and therefore needs dirtying.
+	if (GIsReconstructingBlueprintInstances && !bAllowAutoFindOfManagerInConstructionScript)
+	{
+		return;
+	}
+
+	SetUpPatchManagerData();
+}
+
+void ULandscapePatchComponent::SetUpPatchManagerData()
+{
+	using namespace LandscapePatchComponentLocals;
+
+	if (!IsRealPatch())
 	{
 		return;
 	}
 
 	if (PatchManager.IsValid())
 	{
-		// If we copied over a patch manager, presumably Landscape should be
-		// copied over as well, but might as well do this to be safe.
-		Landscape = PatchManager->GetOwningLandscape();
+		if (Landscape != PatchManager->GetOwningLandscape())
+		{
+			Modify();
+			Landscape = PatchManager->GetOwningLandscape();
+		}
 
+		// The patch manager might legitimately not contain the set patch manager if we're copying a patch. 
 		if (!PatchManager->ContainsPatch(this))
 		{
-			SetPatchManager(PatchManager.Get());
+			// Note: in blueprint-contained patches, this can actually fix a case where our patch manager was
+			// saved incorrectly, but we need the patch manager to be dirtied if we're failing to dirty due to
+			// a load.
+			if (GIsReconstructingBlueprintInstances)
+			{
+				UE_LOG(LogLandscapePatch, Warning, TEXT("Patch was not in the set Patch Manager and was added in a construction script rerun. "
+					"The patch manager should be saved to avoid potential reorderings depending on loading order. (Package: %s, Manager: %s"),
+					*GetNameSafe(PatchManager->GetPackage()), *GetNameSafe(PatchManager.Get()));
+			}
+
+			PatchManager->AddPatch(this);
 		}
+
+		// No more adjustments needed
+		return;
 	}
-	else if (Landscape.IsValid())
+
+	// If we got here, we don't have a patch manager. See if we have a landscape to get it from.
+	if (Landscape.IsValid())
 	{
-		// If we copied over a patch with a landscape but no manager.
 		SetLandscape(Landscape.Get());
+		return;
+	}
+
+	// See if the level has a height patch manager to which we can add ourselves
+	UWorld* World = GetWorld();
+	if (!ensure(World))
+	{
+		return;
+	}
+	TActorIterator<ALandscapePatchManager> ManagerIterator(World);
+	if (!!ManagerIterator)
+	{
+		SetPatchManager(*ManagerIterator);
 	}
 	else
 	{
-		// See if the level has a height patch manager to which we can add ourselves
-		TActorIterator<ALandscapePatchManager> ManagerIterator(World);
-		if (!!ManagerIterator)
-		{
-			SetPatchManager(*ManagerIterator);
-		}
-		else
+		// If no existing manager, find some landscape and add a new one.
+		for (TActorIterator<ALandscape> LandscapeIterator(World); LandscapeIterator; ++LandscapeIterator)
 		{
-			// If no existing manager, find some landscape and add a new one.
-			for (TActorIterator<ALandscape> LandscapeIterator(World); LandscapeIterator; ++LandscapeIterator)
+			if (LandscapeIterator->CanHaveLayersContent())
 			{
-				if (LandscapeIterator->CanHaveLayersContent())
-				{
-					Landscape = *LandscapeIterator;
-					SetPatchManager(CreateNewPatchManagerForLandscape(Landscape.Get()));
-					break;
-				}
-			}
-			if (!PatchManager.IsValid())
-			{
-				UE_LOG(LogLandscapePatch, Warning, TEXT("Unable to find a landscape with edit layers enabled. Unable to create patch manager."));
+				Modify();
+				Landscape = *LandscapeIterator;
+				SetPatchManager(CreateNewPatchManagerForLandscape(Landscape.Get()));
+				break;
 			}
 		}
 	}
+	if (!PatchManager.IsValid())
+	{
+		UE_LOG(LogLandscapePatch, Warning, TEXT("Unable to find a landscape with edit layers enabled. Unable to create patch manager."));
+	}
+}
+
+void ULandscapePatchComponent::MarkDirtyIfModifiedInConstructionScript()
+{
+	if (IsRealPatch() && bDirtiedByConstructionScript)
+	{
+		MarkPackageDirty();
+	}
 }
 
 void ULandscapePatchComponent::PostLoad()
@@ -151,9 +349,17 @@ void ULandscapePatchComponent::PostLoad()
 	bLoadedButNotYetRegistered = true;
 }
 
+void ULandscapePatchComponent::PreSave(FObjectPreSaveContext SaveContext)
+{
+	Super::PreSave(SaveContext);
+
+	// If we're ssaving, then we no longer have to worry about our dirtiness.
+	bDirtiedByConstructionScript = false;
+}
+
 void ULandscapePatchComponent::OnComponentDestroyed(bool bDestroyingHierarchy)
 {
-	if (PatchManager.IsValid())
+	if (PatchManager.IsValid() && !GIsReconstructingBlueprintInstances)
 	{
 		PatchManager->RemovePatch(this);
 	}
@@ -198,22 +404,6 @@ TStructOnScope<FActorComponentInstanceData> ULandscapePatchComponent::GetCompone
 	return MakeStructOnScope<FActorComponentInstanceData, FLandscapePatchComponentInstanceData>(this);
 }
 
-FLandscapePatchComponentInstanceData::FLandscapePatchComponentInstanceData(const ULandscapePatchComponent* Patch)
-: FSceneComponentInstanceData(Patch)
-{
-	ALandscapePatchManager* PatchManager = Patch->GetPatchManager();
-	if (PatchManager)
-	{
-		IndexInManager = PatchManager->GetIndexOfPatch(Patch);
-
-#if WITH_EDITOR
-		bGaveMissingPatchManagerWarning = Patch->bGaveMissingPatchManagerWarning;
-		bGaveNotInPatchManagerWarning = Patch->bGaveNotInPatchManagerWarning;
-		bGaveMissingLandscapeWarning = Patch->bGaveMissingLandscapeWarning;
-#endif
-	}
-}
-
 void ULandscapePatchComponent::OnUpdateTransform(EUpdateTransformFlags UpdateTransformFlags, ETeleportType Teleport)
 {
 	Super::OnUpdateTransform(UpdateTransformFlags, Teleport);
@@ -229,8 +419,7 @@ void ULandscapePatchComponent::PostEditChangeProperty(FPropertyChangedEvent& Pro
 	using namespace LandscapePatchComponentLocals;
 
 	// Do a bunch of checks to make sure that we don't try to do anything when the editing is happening inside the blueprint editor.
-	UWorld* World = GetWorld();
-	if (IsTemplate() || !IsValid(this) || !IsValid(World) || World->WorldType != EWorldType::Editor)
+	if (!IsRealPatch())
 	{
 		return;
 	}
@@ -260,25 +449,120 @@ void ULandscapePatchComponent::PostEditChangeProperty(FPropertyChangedEvent& Pro
 	// the update, etc).
 	Super::PostEditChangeProperty(PropertyChangedEvent);
 }
+#endif // WITH_EDITOR
+
+FLandscapePatchComponentInstanceData::FLandscapePatchComponentInstanceData(const ULandscapePatchComponent* Patch)
+: FSceneComponentInstanceData(Patch)
+{
+#if WITH_EDITOR
+	using namespace LandscapePatchComponentLocals;
+
+	if (!ensure(Patch))
+	{
+		return;
+	}
+
+	bGaveMissingPatchManagerWarning = Patch->bGaveMissingPatchManagerWarning;
+	bGaveNotInPatchManagerWarning = Patch->bGaveNotInPatchManagerWarning;
+	bGaveMissingLandscapeWarning = Patch->bGaveMissingLandscapeWarning;
+
+	bDirtiedByConstructionScript = Patch->bDirtiedByConstructionScript;
+
+	PatchManager = Patch->GetPatchManager();
+	if (PatchManager)
+	{
+		IndexInManager = PatchManager->GetIndexOfPatch(Patch);
+	}
 #endif
+}
 
-// Called after rerunning construction scripts (when patch is part of a blueprint) to make sure that the newly created 
-// instance of the patch is placed in the same order in the patch  manager as the destroyed instance.
+// Called after rerunning construction scripts (when patch is part of a blueprint) to carry over extra data.
 void ULandscapePatchComponent::ApplyComponentInstanceData(FLandscapePatchComponentInstanceData* ComponentInstanceData)
 {
-	if (ComponentInstanceData)
+#if WITH_EDITOR
+	using namespace LandscapePatchComponentLocals;
+
+	if (!ComponentInstanceData || !IsRealPatch())
 	{
-		if (ComponentInstanceData->IndexInManager >= 0 && PatchManager.IsValid())
+		return;
+	}
+
+	bGaveMissingPatchManagerWarning = ComponentInstanceData->bGaveMissingPatchManagerWarning;
+	bGaveNotInPatchManagerWarning = ComponentInstanceData->bGaveNotInPatchManagerWarning;
+	bGaveMissingLandscapeWarning = ComponentInstanceData->bGaveMissingLandscapeWarning;
+
+	bDirtiedByConstructionScript = ComponentInstanceData->bDirtiedByConstructionScript;
+
+	if (!bAllowAutoFindOfManagerInConstructionScript)
+	{
+		// This is the ideal case where we don't have to worry about auto acquired manager and we
+		// don't expect to have to modify the previous one in any way.
+		PatchManager = ComponentInstanceData->PatchManager;
+		PreviousPatchManager = PatchManager;
+		Landscape = PatchManager.IsValid() ? PatchManager->GetOwningLandscape() : nullptr;
+
+		if (PatchManager.IsValid() && ComponentInstanceData->IndexInManager >= 0)
 		{
-			PatchManager->MovePatchToIndex(this, ComponentInstanceData->IndexInManager);
+			int32 CurrentIndexInManager = PatchManager->GetIndexOfPatch(this);
+			if (!ensure(CurrentIndexInManager >= 0))
+			{
+				PatchManager->AddPatch(this);
+				PatchManager->MarkModifiedInConstructionScript();
+			}
+			if (!ensure(CurrentIndexInManager == ComponentInstanceData->IndexInManager))
+			{
+				PatchManager->MovePatchToIndex(this, ComponentInstanceData->IndexInManager);
+				PatchManager->MarkModifiedInConstructionScript();
+			}
 		}
+		return;
+	}
 
-#if WITH_EDITOR
-		bGaveMissingPatchManagerWarning = ComponentInstanceData->bGaveMissingPatchManagerWarning;
-		bGaveNotInPatchManagerWarning = ComponentInstanceData->bGaveNotInPatchManagerWarning;
-		bGaveMissingLandscapeWarning = ComponentInstanceData->bGaveMissingLandscapeWarning;
-#endif
+	// If we got here, then we're going down a legacy path where we auto acquire patch manager
+	// in construction script.
+
+	if (ComponentInstanceData->PatchManager.IsValid())
+	{
+		// We had an existing patch manager so use that. This is mostly like the above except that we
+		// have to remove ourselves from an acquired manager if there is one, and we automatically
+		// add ourselves to our patch manager if it was saved without this patch. Calling SetPatchManager
+		// does all of that for us, except we need to track whether we modified the patch manager we keep.
+
+		bool bModifiedPatchManager = !ComponentInstanceData->PatchManager->ContainsPatch(this) 
+			|| ComponentInstanceData->PatchManager->GetIndexOfPatch(this) != ComponentInstanceData->IndexInManager;
+
+		SetPatchManager(ComponentInstanceData->PatchManager.Get());
+
+		PatchManager->MovePatchToIndex(this, ComponentInstanceData->IndexInManager);
+
+		if (bModifiedPatchManager)
+		{
+			ComponentInstanceData->PatchManager->MarkModifiedInConstructionScript();
+
+			// This might get double logged because InstanceDataCache->ApplyToActor gets run twice (see AActor::ExecuteConstruction),
+			// but that's ok. Also we could check for current dirtyness to avoid the "might not have" wording,
+			// but we want to log a warning either way, so probably not worth it.
+			UE_LOG(LogLandscapePatch, Warning, TEXT("Patch was missing from expected patch manager and was added in a "
+				"construction script rerun but manager might not have been marked dirty. Save the manager to avoid "
+				"potential reorderings depending on loading order. (Package: %s, Actor: %s"),
+				*GetNameSafe(GetPackage()), *GetNameSafe(GetAttachmentRootActor()));
+		}
+		return;
 	}
+
+	// If we got to here, then we're keeping the patch manager we auto acquired.
+	if (PatchManager.IsValid())
+	{
+		bDirtiedByConstructionScript = true;
+		PatchManager->MarkModifiedInConstructionScript();
+
+		// See comment on the other log statement above
+		UE_LOG(LogLandscapePatch, Warning, TEXT("Patch acquired a patch manager in a construction script rerun but might not have been "
+			"marked dirty. Save the patch and its manager to avoid potential reorderings depending on loading order. (Package: %s, Actor: %s"),
+			*GetNameSafe(GetPackage()), *GetNameSafe(GetAttachmentRootActor()));
+	}
+	
+#endif // WITH_EDITOR
 }
 
 void ULandscapePatchComponent::SetLandscape(ALandscape* NewLandscape)
@@ -295,6 +579,8 @@ void ULandscapePatchComponent::SetLandscape(ALandscape* NewLandscape)
 		return;
 	}
 
+	Modify();
+
 	Landscape = NewLandscape;
 
 	if (!NewLandscape)
@@ -303,8 +589,7 @@ void ULandscapePatchComponent::SetLandscape(ALandscape* NewLandscape)
 		return;
 	}
 
-	UWorld* World = GetWorld();
-	if (IsTemplate() || !IsValid(World) || World->WorldType != EWorldType::Editor)
+	if (!IsRealPatch())
 	{
 		return;
 	}
@@ -336,10 +621,13 @@ void ULandscapePatchComponent::SetPatchManager(ALandscapePatchManager* NewPatchM
 	// surround everything since nothing is currently expected to work for a patch manager at runtime.
 #if WITH_EDITOR
 
-	if (PreviousPatchManager == NewPatchManager && NewPatchManager && NewPatchManager->ContainsPatch(this))
+	if (PreviousPatchManager == NewPatchManager 
+		&& PatchManager == NewPatchManager
+		&& !(NewPatchManager && !NewPatchManager->ContainsPatch(this)))
 	{
 		return;
 	}
+
 	ResetWarnings();
 
 	if (PreviousPatchManager.IsValid())
@@ -347,11 +635,11 @@ void ULandscapePatchComponent::SetPatchManager(ALandscapePatchManager* NewPatchM
 		PreviousPatchManager->RemovePatch(this);
 	}
 
+	Modify();
 	PatchManager = NewPatchManager;
 	if (NewPatchManager)
 	{
-		UWorld* World = GetWorld();
-		if (!IsTemplate() && IsValid(World) && World->WorldType == EWorldType::Editor)
+		if (IsRealPatch())
 		{
 			PatchManager->AddPatch(this);
 		}
@@ -388,34 +676,25 @@ void ULandscapePatchComponent::MoveToTop()
 
 void ULandscapePatchComponent::RequestLandscapeUpdate(bool bInUserTriggeredUpdate)
 {
+	using namespace LandscapePatchComponentLocals;
+
 #if WITH_EDITOR
-	UWorld* World = GetWorld();
 	// Note that aside from the usual guard against doing things in the blueprint editor, the check of WorldType
-	// here also prevents us from doing the request while cooking, where WorldType is set to Inactive. Otherwise
+	// inside this call also prevents us from doing the request while cooking, where WorldType is set to Inactive. Otherwise
 	// we would issue a warning below while cooking when PatchManager is not saved (and not reset in the construction
 	// script).
-	if (IsTemplate() || World->WorldType != EWorldType::Editor)
+	if (!IsRealPatch())
 	{
 		return;
 	}
 
-	// Helper for getting name strings in the warning messages below.
-	auto GetNameString = [](UObject* PotentiallyNullObject) {
-		FString ToReturn;
-		if (PotentiallyNullObject)
-		{
-			PotentiallyNullObject->GetName(ToReturn);
-		}
-		return ToReturn;
-	};
-
 	if (!PatchManager.IsValid())
 	{
 		if (!bGaveMissingPatchManagerWarning)
 		{
 			UE_LOG(LogLandscapePatch, Warning, TEXT("Patch does not have a valid patch manager. "
 				"Set the landscape or patch manager on the patch. (Package: %s, Actor: %s"), 
-				*GetNameString(GetPackage()), *GetNameString(GetAttachmentRootActor()));
+				*GetNameSafe(GetPackage()), *GetNameSafe(GetAttachmentRootActor()));
 			bGaveMissingPatchManagerWarning = true;
 		}
 		return;
@@ -428,7 +707,7 @@ void ULandscapePatchComponent::RequestLandscapeUpdate(bool bInUserTriggeredUpdat
 		{
 			UE_LOG(LogLandscapePatch, Warning, TEXT("Patch's patch manager does not contain this patch. "
 				"Perhaps the manager was not saved? Reset the patch manager on the patch. (Package: %s, Actor: %s)"), 
-				*GetNameString(GetPackage()), *GetNameString(GetAttachmentRootActor()));
+				*GetNameSafe(GetPackage()), *GetNameSafe(GetAttachmentRootActor()));
 			bGaveNotInPatchManagerWarning = true;
 		}
 		bRequestUpdate = false;
@@ -439,7 +718,7 @@ void ULandscapePatchComponent::RequestLandscapeUpdate(bool bInUserTriggeredUpdat
 		{
 			UE_LOG(LogLandscapePatch, Warning, TEXT("Patch's patch manager does not have a valid owning "
 				"landscape. Perhaps the landscape was not saved? Reset the landscape on the manager. (Package: %s, Manager: %s)"), 
-				*GetNameString(GetPackage()), *GetNameString(PatchManager.Get()));
+				*GetNameSafe(GetPackage()), *GetNameSafe(PatchManager.Get()));
 			bGaveMissingLandscapeWarning = true;
 		}
 		bRequestUpdate = false;
@@ -462,6 +741,12 @@ void ULandscapePatchComponent::ResetWarnings()
 }
 #endif
 
+bool ULandscapePatchComponent::IsRealPatch()
+{
+	UWorld* World = GetWorld();
+	return !IsTemplate() && IsValid(this) && IsValid(World) && World->WorldType == EWorldType::Editor;
+}
+
 FTransform ULandscapePatchComponent::GetLandscapeHeightmapCoordsToWorld() const
 {
 	if (PatchManager.IsValid())
diff --git a/Engine/Plugins/Experimental/LandscapePatch/Source/LandscapePatch/Private/LandscapePatchManager.cpp b/Engine/Plugins/Experimental/LandscapePatch/Source/LandscapePatch/Private/LandscapePatchManager.cpp
--- a/Engine/Plugins/Experimental/LandscapePatch/Source/LandscapePatch/Private/LandscapePatchManager.cpp
+++ b/Engine/Plugins/Experimental/LandscapePatch/Source/LandscapePatch/Private/LandscapePatchManager.cpp
@@ -2,9 +2,10 @@
 
 #include "LandscapePatchManager.h"
 
-#include "Modules/ModuleManager.h"
+#include "CoreGlobals.h" // GIsReconstructingBlueprintInstances
 #include "Engine/Level.h"
 #include "Engine/World.h"
+#include "HAL/IConsoleManager.h" // FAutoConsoleCommand
 #include "Landscape.h"
 #include "LandscapeEditTypes.h"
 #include "LandscapeDataAccess.h"
@@ -12,11 +13,54 @@
 #include "LandscapePatchLogging.h"
 #include "LandscapeModule.h"
 #include "LandscapeEditorServices.h"
+#include "Logging/MessageLog.h"
+#include "Misc/UObjectToken.h"
+#include "Modules/ModuleManager.h"
+#include "UObject/ObjectSaveContext.h"
+#include "UObject/Package.h"
+#include "UObject/UObjectBaseUtility.h" // GetNameSafe
+#include "UObject/UObjectIterator.h"
 
 #include UE_INLINE_GENERATED_CPP_BY_NAME(LandscapePatchManager)
 
 #define LOCTEXT_NAMESPACE "LandscapePatchManager"
 
+namespace LandscapePatchManagerLocals
+{
+#if WITH_EDITOR
+	// Another way to do the patch pointer fixup in addition to the map check.
+	FAutoConsoleCommand CCmdFixOwnedPatchBackPointers(
+		TEXT("LandscapePatch.PatchManager.FixOwnedPatchBackPointers"),
+		TEXT("For all patch managers, make it so that their managed patches point to them as their "
+			"patch manager. This is a more limited cleanup command than LandscapePatch.FixPatchManagerPointers."),
+		FConsoleCommandDelegate::CreateLambda([]() 
+	{
+		for (TObjectIterator<ALandscapePatchManager> It(
+			/*AdditionalExclusionFlags = */RF_ClassDefaultObject,
+			/*bIncludeDerivedClasses = */true,
+			/*InInternalExclusionFlags = */EInternalObjectFlags::Garbage); It; ++It)
+		{
+			ALandscapePatchManager* Manager = *It;
+			if (!IsValid(Manager))
+			{
+				continue;
+			}
+
+			UWorld* World = Manager->GetWorld();
+			if (Manager->IsTemplate() || !IsValid(World) || World->WorldType != EWorldType::Editor)
+			{
+				continue;
+			}
+
+			if (Manager)
+			{
+				Manager->FixOwnedPatchBackPointers();
+			}
+		}
+	}));
+#endif // WITH_EDITOR
+}
+
 // TODO: Not sure if using this kind of constructor is a proper thing to do vs some other hook...
 ALandscapePatchManager::ALandscapePatchManager(const FObjectInitializer& ObjectInitializer)
 	: ALandscapeBlueprintBrushBase(ObjectInitializer)
@@ -74,6 +118,8 @@ void ALandscapePatchManager::Initialize_Native(const FTransform & InLandscapeTra
 
 UTextureRenderTarget2D* ALandscapePatchManager::RenderLayer_Native(const FLandscapeBrushParameters& InParameters)
 {
+	using namespace LandscapePatchManagerLocals;
+
 	// Used to determine whether we need to remove any invalid brushes
 	bool bHaveInvalidPatches = false;
 	FLandscapeBrushParameters BrushParameters = InParameters;
@@ -92,14 +138,7 @@ UTextureRenderTarget2D* ALandscapePatchManager::RenderLayer_Native(const FLandsc
 			Component.LoadSynchronous();
 		}
 
-		if (Component.IsValid())
-		{
-			if (Component->IsEnabled())
-			{
-				BrushParameters.CombinedResult = Component->RenderLayer_Native(BrushParameters);
-			}
-		}
-		else if (Component.IsNull())
+		if (Component.IsNull())
 		{
 			// Theoretically when components are marked for destruction, they should remove themselves from
 			// the patch manager in their OnComponentDestroyed call. However there seem to be ways to end up
@@ -107,8 +146,10 @@ UTextureRenderTarget2D* ALandscapePatchManager::RenderLayer_Native(const FLandsc
 			// patch actor.
 			UE_LOG(LogLandscapePatch, Warning, TEXT("ALandscapePatchManager: Found an invalid patch in patch manager. It will be removed."));
 			bHaveInvalidPatches = true;
+			continue;
 		}
-		else
+
+		if (!Component.IsValid())
 		{
 			// This means that IsPending() was true, but LoadSynchronous() failed, which we generally don't
 			// expect to happen. However, it can happen in some edge cases such as if you force delete a patch
@@ -118,7 +159,33 @@ UTextureRenderTarget2D* ALandscapePatchManager::RenderLayer_Native(const FLandsc
 				"turned out to be invalid. It will be removed."));
 			Component = nullptr;
 			bHaveInvalidPatches = true;
+			continue;
+		}
+
+		if (!Component->IsEnabled())
+		{
+			// Skip disabled patches
+			continue;
 		}
+
+		// Make sure the patch has this manager set as its patch manager.
+		if (Component->GetPatchManager() != this)
+		{
+#if WITH_EDITORONLY_DATA
+			if (!bIssuedPatchOwnershipWarning)
+			{
+				bIssuedPatchOwnershipWarning = true;
+				UE_LOG(LogLandscapePatch, Warning, TEXT("ALandscapePatchManager: Found at least one patch whose patch manager is not "
+					"set to a patch manager that contains it. This should be fixed an the patches resaved (do this through the map check "
+					"action, or via LandscapePatch.PatchManager.FixOwnedPatchBackPointers console command.)"));
+			}
+#endif
+
+			// The patches need a proper back pointer to get parameters for their render call.
+			continue;
+		}
+
+		BrushParameters.CombinedResult = Component->RenderLayer_Native(BrushParameters); 
 	}
 
 	if (bHaveInvalidPatches)
@@ -187,10 +254,22 @@ bool ALandscapePatchManager::ContainsPatch(ULandscapePatchComponent* Patch) cons
 
 void ALandscapePatchManager::AddPatch(ULandscapePatchComponent* Patch)
 {
+	using namespace LandscapePatchManagerLocals;
+
 	if (Patch)
 	{
-		Modify();
-		PatchComponents.AddUnique(TSoftObjectPtr<ULandscapePatchComponent>(Patch));
+		if (!ContainsPatch(Patch))
+		{
+			Modify();
+			PatchComponents.Add(TSoftObjectPtr<ULandscapePatchComponent>(Patch));
+		}
+
+		if (Patch->GetPatchManager() != this)
+		{
+			UE_LOG(LogLandscapePatch, Warning, TEXT("ALandscapePatchManager::AddPatch: Added patch does not have this manager set "
+				"as its manager. Patches should typically be added to managers by setting the manager on the patch. "
+				"(Package: % s, Actor : % s)"), *GetNameSafe(Patch->GetPackage()), *GetNameSafe(Patch->GetAttachmentRootActor()));
+		}
 
 		// No need to update if the patch is disabled. Important to avoid needlessly updating while dragging a blueprint with
 		// a disabled patch (since construction scripts constantly add and remove).
@@ -205,7 +284,7 @@ bool ALandscapePatchManager::RemovePatch(ULandscapePatchComponent* Patch)
 {
 	bool bRemoved = false;
 
-	if (Patch)
+	if (Patch && ContainsPatch(Patch))
 	{
 		Modify();
 		bRemoved = PatchComponents.Remove(TSoftObjectPtr<ULandscapePatchComponent>(Patch)) > 0;
@@ -233,6 +312,8 @@ void ALandscapePatchManager::MovePatchToIndex(ULandscapePatchComponent* Patch, i
 		return;
 	}
 
+	Modify();
+
 	// It might seem like the index needs adjusting if we're removing before the given index, but that
 	// is not the case if our goal is for the index of the patch to be the given index at the end (rather
 	// than our goal being that the patch be in a particular position relative to the existing patches).
@@ -248,6 +329,30 @@ void ALandscapePatchManager::MovePatchToIndex(ULandscapePatchComponent* Patch, i
 }
 
 #if WITH_EDITOR
+void ALandscapePatchManager::FixOwnedPatchBackPointers()
+{
+	for (TSoftObjectPtr<ULandscapePatchComponent>& Component : PatchComponents)
+	{
+		if (Component.IsValid() && Component->GetPatchManager() != this)
+		{
+			Component->SetPatchManager(this);
+		}
+	}
+}
+
+void ALandscapePatchManager::MarkModifiedInConstructionScript()
+{
+	bDirtiedByConstructionScript = true;
+}
+
+void ALandscapePatchManager::MarkDirtyIfModifiedInConstructionScript()
+{
+	if (bDirtiedByConstructionScript)
+	{
+		MarkPackageDirty();
+	}
+}
+
 PRAGMA_DISABLE_DEPRECATION_WARNINGS
 bool ALandscapePatchManager::IsAffectingWeightmapLayer(const FName& InLayerName) const
 {
@@ -306,6 +411,14 @@ void ALandscapePatchManager::PostEditUndo()
 	RequestLandscapeUpdate();
 }
 
+void ALandscapePatchManager::PreSave(FObjectPreSaveContext SaveContext)
+{
+	Super::PreSave(SaveContext);
+
+	// If we're saving, then we no longer have to worry about our dirtiness.
+	bDirtiedByConstructionScript = false;
+}
+
 void ALandscapePatchManager::SetOwningLandscape(ALandscape* InOwningLandscape)
 {
 	Super::SetOwningLandscape(InOwningLandscape);
@@ -331,6 +444,65 @@ void ALandscapePatchManager::PostEditChangeProperty(FPropertyChangedEvent& Prope
 		SetTargetLandscape(DetailPanelLandscape.Get());
 	}
 }
-#endif
+
+void ALandscapePatchManager::CheckForErrors()
+{
+	using namespace LandscapePatchManagerLocals;
+
+	Super::CheckForErrors();
+
+	auto GetPackageAndActorArgs = [this]()
+	{
+		FFormatNamedArguments Arguments;
+		Arguments.Add(TEXT("Package"), FText::FromString(*GetNameSafe(GetPackage())));
+		Arguments.Add(TEXT("Actor"), FText::FromString(*GetNameSafe(this)));
+		return Arguments;
+	};
+
+	// See if we're supposed to be dirty but aren't
+	if (bDirtiedByConstructionScript && !(GetPackage() && GetPackage()->IsDirty()))
+	{
+		FMessageLog("MapCheck").Warning()
+			->AddToken(FUObjectToken::Create(this))
+			->AddToken(FTextToken::Create(FText::Format(LOCTEXT("ManagerNeedsSaving", "Manager got modified in a construction "
+				"script rerun but has not been marked dirty. This could cause patch order to be unstable. "
+				"(Package: {Package}, Manager: {Actor})."), GetPackageAndActorArgs())))
+			->AddToken(FActionToken::Create(LOCTEXT("MarkDirtyButton", "Mark dirty"), FText(),
+				FOnActionTokenExecuted::CreateWeakLambda(this, [this]()
+		{
+			MarkPackageDirty();
+		})));
+	}
+
+	// See if we're holding on to any patches that don't have us as the owning patch manager
+	bool bHavePatchWithIncorrectManager = false;
+	for (TSoftObjectPtr<ULandscapePatchComponent>& Component : PatchComponents)
+	{
+		if (Component.IsValid() && Component->GetPatchManager() != this)
+		{
+			bHavePatchWithIncorrectManager = true;
+			break;
+		}
+	}
+
+	if (bHavePatchWithIncorrectManager)
+	{
+		FFormatNamedArguments Arguments;
+		Arguments.Add(TEXT("Manager"), FText::FromString(GetActorNameOrLabel()));
+		Arguments.Add(TEXT("Outer"), FText::FromString(GetPackage()->GetPathName()));
+
+		FMessageLog("MapCheck").Warning()
+			->AddToken(FUObjectToken::Create(this))
+			->AddToken(FTextToken::Create(FText::Format(LOCTEXT("PatchesHaveIncorrectManagerPointer", "Patch manager holds at "
+				"least one patch whose patch manager pointer is set incorrectly. These need to be fixed and resaved. "
+				"(Package: {Package}, Manager: {Actor})."), GetPackageAndActorArgs())))
+			->AddToken(FActionToken::Create(LOCTEXT("FixPatchesButton", "Fix patches"), FText(),
+				FOnActionTokenExecuted::CreateWeakLambda(this, [this]() 
+		{
+			FixOwnedPatchBackPointers();
+		})));
+	}
+}
+#endif // WITH_EDITOR
 
 #undef LOCTEXT_NAMESPACE
diff --git a/Engine/Plugins/Experimental/LandscapePatch/Source/LandscapePatch/Private/LandscapeTexturePatch.cpp b/Engine/Plugins/Experimental/LandscapePatch/Source/LandscapePatch/Private/LandscapeTexturePatch.cpp
--- a/Engine/Plugins/Experimental/LandscapePatch/Source/LandscapePatch/Private/LandscapeTexturePatch.cpp
+++ b/Engine/Plugins/Experimental/LandscapePatch/Source/LandscapePatch/Private/LandscapeTexturePatch.cpp
@@ -79,6 +79,16 @@ namespace LandscapeTexturePatchLocals
 UTextureRenderTarget2D* ULandscapeTexturePatch::RenderLayer_Native(const FLandscapeBrushParameters& InParameters)
 {
 	using namespace UE::Landscape;
+
+	// If we're getting a RenderLayer_Native call, then we're inside some patch manager, and we
+	// expect our pointer to be properly initialized. This assumption could be violated if we
+	// were not consistent in saving both the manager and the patch. However the manager should
+	// be catching this case for us.
+	if (!ensure(PatchManager.IsValid()))
+	{
+		return InParameters.CombinedResult;
+	}
+
 	const bool bIsHeightmapTarget = InParameters.LayerType == ELandscapeToolTargetType::Heightmap;
 	const bool bIsWeightmapTarget = InParameters.LayerType == ELandscapeToolTargetType::Weightmap;
 	const bool bIsVisibilityLayerTarget = InParameters.LayerType == ELandscapeToolTargetType::Visibility;
diff --git a/Engine/Plugins/Experimental/LandscapePatch/Source/LandscapePatch/Public/LandscapePatchComponent.h b/Engine/Plugins/Experimental/LandscapePatch/Source/LandscapePatch/Public/LandscapePatchComponent.h
--- a/Engine/Plugins/Experimental/LandscapePatch/Source/LandscapePatch/Public/LandscapePatchComponent.h
+++ b/Engine/Plugins/Experimental/LandscapePatch/Source/LandscapePatch/Public/LandscapePatchComponent.h
@@ -76,10 +76,23 @@ public:
 	// so that they can be called from non-editor-only classes in editor contexts.
 #if WITH_EDITOR
 
+	/**
+	 * Verifies that the PatchManager and Landscape pointers are set up properly and alters them if needed. Public so
+	 * that it can be called from a console command.
+	 */
+	void SetUpPatchManagerData();
+
+	/**
+	 * Dirties the patch if it was modified in a construction script but was unable to mark itself
+	 * dirty. Meant to be used by cleanup commands.
+	 */
+	void MarkDirtyIfModifiedInConstructionScript();
+
 	// USceneComponent
 	virtual void OnUpdateTransform(EUpdateTransformFlags UpdateTransformFlags, ETeleportType Teleport) override;
 
 	// UActorComponent
+	virtual void CheckForErrors() override;
 	virtual void OnComponentCreated() override;
 	virtual void OnComponentDestroyed(bool bDestroyingHierarchy) override;
 	virtual void OnRegister() override;
@@ -89,12 +102,12 @@ public:
 	// UObject
 	virtual void PostEditChangeProperty(FPropertyChangedEvent& PropertyChangedEvent) override;
 	virtual bool IsPostLoadThreadSafe() const override { return true; }
-	virtual void PostLoad();
+	virtual void PostLoad() override;
+	virtual void PreSave(FObjectPreSaveContext SaveContext) override;
 #endif
 	virtual bool IsEditorOnly() const override { return true; }
 	virtual bool NeedsLoadForClient() const override { return false; }
 	virtual bool NeedsLoadForServer() const override { return false; }
-
 protected:
 	/**
 	 * Move the patch to be the last processed patch in the current patch manager. This is a way to
@@ -132,12 +145,20 @@ private:
 	UPROPERTY()
 	TSoftObjectPtr<ALandscapePatchManager> PreviousPatchManager = nullptr;
 
+	// Verifies that we're not an archetype, trash, or not in a proper world
+	bool IsRealPatch();
+
 #if WITH_EDITOR
 	// Used to avoid spamming warning messages
 	bool bGaveMissingPatchManagerWarning = false;
 	bool bGaveNotInPatchManagerWarning = false;
 	bool bGaveMissingLandscapeWarning = false;
 	void ResetWarnings();
+
+	// Automatic patch registration during on-load construction script reruns could end up modifying
+	// the patch without dirtying. This is dangerous as it can result in unstable patch ordering,
+	// so we want to detect this case.
+	bool bDirtiedByConstructionScript = false;
 #endif
 
 	friend struct FLandscapePatchComponentInstanceData;
@@ -165,10 +186,17 @@ struct FLandscapePatchComponentInstanceData : public FSceneComponentInstanceData
 		CastChecked<ULandscapePatchComponent>(Component)->ApplyComponentInstanceData(this);
 	}
 
+#if WITH_EDITORONLY_DATA
+	// The UPROPERTY tags inside a FSceneComponentInstanceData might not be necessary, but might 
+	// potentially be used in some multiuser code paths.
+
+	UPROPERTY()
+	TSoftObjectPtr<ALandscapePatchManager> PatchManager = nullptr;
 	UPROPERTY()
 	int32 IndexInManager = -1;
-
-#if WITH_EDITORONLY_DATA
+	
+	UPROPERTY()
+	bool bDirtiedByConstructionScript = false;
 
 	// Used so that we don't spam warning messages while rerunning construction scripts on a patch
 	// that triggers one of the warnings.
diff --git a/Engine/Plugins/Experimental/LandscapePatch/Source/LandscapePatch/Public/LandscapePatchManager.h b/Engine/Plugins/Experimental/LandscapePatch/Source/LandscapePatch/Public/LandscapePatchManager.h
--- a/Engine/Plugins/Experimental/LandscapePatch/Source/LandscapePatch/Public/LandscapePatchManager.h
+++ b/Engine/Plugins/Experimental/LandscapePatch/Source/LandscapePatch/Public/LandscapePatchManager.h
@@ -69,8 +69,26 @@ public:
 	 */
 	UFUNCTION(BlueprintCallable, Category = LandscapePatch)
 	void MovePatchToIndex(ULandscapePatchComponent* Patch, int32 Index);
-	
+
 #if WITH_EDITOR
+	/**
+	 * A helper cleanup method to fix things if something goes wrong in saving and owned patches do not have
+	 * the correct patch manager pointer back. Public so that it can be called from a console command.
+	 */
+	void FixOwnedPatchBackPointers();
+
+	/**
+	 * Marks that the patch manager was modified during a construction script rerun where it might 
+	 * not be able to mark itself dirty (if it was done during loading).
+	 */
+	void MarkModifiedInConstructionScript();
+
+	/**
+	 * Dirties the manager if it was modified in a construction script but was unable to mark itself
+	 * dirty. Meant to be used by cleanup commands.
+	 */
+	void MarkDirtyIfModifiedInConstructionScript();
+
 	// ALandscapeBlueprintBrushBase
 	UE_DEPRECATED(5.3, "Use AffectsWeightmapLayer")
 	virtual bool IsAffectingWeightmapLayer(const FName& InLayerName) const override;
@@ -78,8 +96,12 @@ public:
 	virtual bool AffectsVisibilityLayer() const override;
 	virtual void SetOwningLandscape(class ALandscape* InOwningLandscape) override;
 
+	// AActor
+	virtual void CheckForErrors() override;
+
 	// UObject
 	virtual void PostEditUndo() override;
+	virtual void PreSave(FObjectPreSaveContext SaveContext) override;
 	virtual void PostEditChangeProperty(FPropertyChangedEvent& PropertyChangedEvent) override;
 #endif
 	virtual bool IsEditorOnly() const override { return true; }
@@ -101,6 +123,15 @@ protected:
 	 */
 	UPROPERTY(EditAnywhere, Category = Landscape, Transient, meta = (DisplayName = "Landscape"))
 	TObjectPtr<ALandscape> DetailPanelLandscape = nullptr;
+
+private:
+	bool bIssuedPatchOwnershipWarning = false;
+
+	// The interaction of automatic patch registration and construction script reruns could end
+	// up modifying the manager during a load if things had to be fixed up, but the manager might
+	// not end up being marked dirty. This is dangerous as it can result in unstable patch ordering,
+	// so we want to detect this case.
+	bool bDirtiedByConstructionScript = false;
 #endif
 };
 
diff --git a/Engine/Plugins/Experimental/LandscapePatch/Source/LandscapePatchEditorOnly/Private/LandscapeTexturePatchVisualizer.cpp b/Engine/Plugins/Experimental/LandscapePatch/Source/LandscapePatchEditorOnly/Private/LandscapeTexturePatchVisualizer.cpp
--- a/Engine/Plugins/Experimental/LandscapePatch/Source/LandscapePatchEditorOnly/Private/LandscapeTexturePatchVisualizer.cpp
+++ b/Engine/Plugins/Experimental/LandscapePatch/Source/LandscapePatchEditorOnly/Private/LandscapeTexturePatchVisualizer.cpp
@@ -20,8 +20,6 @@ void FLandscapeTexturePatchVisualizer::DrawVisualization(const UActorComponent*
 			Color, Patch->GetUnscaledCoverage().X * PatchToWorld.GetScale3D().X, Patch->GetUnscaledCoverage().Y * PatchToWorld.GetScale3D().Y, SDPG_Foreground,
 			Thickness, DepthBias, bScreenSpace);
 	}
-	else
-	{
-		ensure(false);
-	}
+	// Note: when patches are deleted, selection seems to be updated only after giving another call to the
+	// visualization drawing, so do not ensure if patch was null.
 }
