LandscapePatch: Remove incorrect beta tag on plugin (plugin is already marked experimental and is in experimental folder).

diff --git a/Engine/Plugins/Experimental/LandscapePatch/LandscapePatch.uplugin b/Engine/Plugins/Experimental/LandscapePatch/LandscapePatch.uplugin
--- a/Engine/Plugins/Experimental/LandscapePatch/LandscapePatch.uplugin
+++ b/Engine/Plugins/Experimental/LandscapePatch/LandscapePatch.uplugin
@@ -11,7 +11,6 @@
 	"MarketplaceURL": "",
 	"SupportURL": "",
 	"CanContainContent": true,
-	"IsBetaVersion": true,
 	"IsExperimentalVersion": true,
 	"Installed": false,
 	"Modules": [
