LandscapePatch: Add index lookup helper to patch manager to speed up Contains and IndexOf queries.

Also stop trying to store index in blueprint actors across construction script reruns because it should no longer be necessary now that patches do not detach across a rerun, and because keeping them in instance data can cause a subtle bug in undoing patch deletion: a BP actor can get captured in an unrelated deletion transaction due to one of its components being garbage collected, and an incorrect index would be stored and applied because the index is stored after deletion but applied after undoing deletion.

diff --git a/Engine/Plugins/Experimental/LandscapePatch/Source/LandscapePatch/Private/LandscapePatchComponent.cpp b/Engine/Plugins/Experimental/LandscapePatch/Source/LandscapePatch/Private/LandscapePatchComponent.cpp
--- a/Engine/Plugins/Experimental/LandscapePatch/Source/LandscapePatch/Private/LandscapePatchComponent.cpp
+++ b/Engine/Plugins/Experimental/LandscapePatch/Source/LandscapePatch/Private/LandscapePatchComponent.cpp
@@ -469,10 +469,6 @@ FLandscapePatchComponentInstanceData::FLandscapePatchComponentInstanceData(const
 	bDirtiedByConstructionScript = Patch->bDirtiedByConstructionScript;
 
 	PatchManager = Patch->GetPatchManager();
-	if (PatchManager)
-	{
-		IndexInManager = PatchManager->GetIndexOfPatch(Patch);
-	}
 #endif
 }
 
@@ -501,19 +497,13 @@ void ULandscapePatchComponent::ApplyComponentInstanceData(FLandscapePatchCompone
 		PreviousPatchManager = PatchManager;
 		Landscape = PatchManager.IsValid() ? PatchManager->GetOwningLandscape() : nullptr;
 
-		if (PatchManager.IsValid() && ComponentInstanceData->IndexInManager >= 0)
+		if (PatchManager.IsValid())
 		{
-			int32 CurrentIndexInManager = PatchManager->GetIndexOfPatch(this);
-			if (!ensure(CurrentIndexInManager >= 0))
+			if (!ensure(PatchManager->ContainsPatch(this)))
 			{
 				PatchManager->AddPatch(this);
 				PatchManager->MarkModifiedInConstructionScript();
 			}
-			if (!ensure(CurrentIndexInManager == ComponentInstanceData->IndexInManager))
-			{
-				PatchManager->MovePatchToIndex(this, ComponentInstanceData->IndexInManager);
-				PatchManager->MarkModifiedInConstructionScript();
-			}
 		}
 		return;
 	}
@@ -528,13 +518,10 @@ void ULandscapePatchComponent::ApplyComponentInstanceData(FLandscapePatchCompone
 		// add ourselves to our patch manager if it was saved without this patch. Calling SetPatchManager
 		// does all of that for us, except we need to track whether we modified the patch manager we keep.
 
-		bool bModifiedPatchManager = !ComponentInstanceData->PatchManager->ContainsPatch(this) 
-			|| ComponentInstanceData->PatchManager->GetIndexOfPatch(this) != ComponentInstanceData->IndexInManager;
+		bool bModifiedPatchManager = !ComponentInstanceData->PatchManager->ContainsPatch(this);
 
 		SetPatchManager(ComponentInstanceData->PatchManager.Get());
 
-		PatchManager->MovePatchToIndex(this, ComponentInstanceData->IndexInManager);
-
 		if (bModifiedPatchManager)
 		{
 			ComponentInstanceData->PatchManager->MarkModifiedInConstructionScript();
diff --git a/Engine/Plugins/Experimental/LandscapePatch/Source/LandscapePatch/Private/LandscapePatchManager.cpp b/Engine/Plugins/Experimental/LandscapePatch/Source/LandscapePatch/Private/LandscapePatchManager.cpp
--- a/Engine/Plugins/Experimental/LandscapePatch/Source/LandscapePatch/Private/LandscapePatchManager.cpp
+++ b/Engine/Plugins/Experimental/LandscapePatch/Source/LandscapePatch/Private/LandscapePatchManager.cpp
@@ -120,8 +120,19 @@ UTextureRenderTarget2D* ALandscapePatchManager::RenderLayer_Native(const FLandsc
 {
 	using namespace LandscapePatchManagerLocals;
 
-	// Used to determine whether we need to remove any invalid brushes
+	// Used for removing invalid brushes. We remove from the index map immediately but then remove
+	// from the array and update other indices at the very end.
 	bool bHaveInvalidPatches = false;
+	int32 MinRemovedIndex = PatchComponents.Num();
+	auto RemoveComponentFromIndexMap = [this, &MinRemovedIndex](TSoftObjectPtr<ULandscapePatchComponent>& Component)
+	{
+		int32 RemovedIndex = -1;
+		if (PatchToIndex.RemoveAndCopyValue(Component, RemovedIndex))
+		{
+			MinRemovedIndex = FMath::Min(MinRemovedIndex, RemovedIndex);
+		}
+	};
+
 	FLandscapeBrushParameters BrushParameters = InParameters;
 
 	// TODO: There are many uncertainties in how we iterate across the height patches and have them
@@ -145,6 +156,7 @@ UTextureRenderTarget2D* ALandscapePatchManager::RenderLayer_Native(const FLandsc
 			// with destroyed patches not being removed, for instance through saving the manager but not the
 			// patch actor.
 			UE_LOG(LogLandscapePatch, Warning, TEXT("ALandscapePatchManager: Found an invalid patch in patch manager. It will be removed."));
+			RemoveComponentFromIndexMap(Component);
 			bHaveInvalidPatches = true;
 			continue;
 		}
@@ -157,6 +169,7 @@ UTextureRenderTarget2D* ALandscapePatchManager::RenderLayer_Native(const FLandsc
 			// a dead patch that actually needs removal.
 			UE_LOG(LogLandscapePatch, Warning, TEXT("ALandscapePatchManager: Found a pending patch pointer in patch manager that "
 				"turned out to be invalid. It will be removed."));
+			RemoveComponentFromIndexMap(Component);
 			Component = nullptr;
 			bHaveInvalidPatches = true;
 			continue;
@@ -192,7 +205,12 @@ UTextureRenderTarget2D* ALandscapePatchManager::RenderLayer_Native(const FLandsc
 	{
 		PatchComponents.RemoveAll([](TSoftObjectPtr<ULandscapePatchComponent> Component) {
 			return Component.IsNull();
-			});
+		});
+		// Update forward indices
+		for (int32 i = MinRemovedIndex; i < PatchComponents.Num(); ++i)
+		{
+			PatchToIndex.Add(PatchComponents[i], i);
+		}
 	}
 
 	return BrushParameters.CombinedResult;
@@ -249,7 +267,7 @@ void ALandscapePatchManager::SetTargetLandscape(ALandscape* InTargetLandscape)
 
 bool ALandscapePatchManager::ContainsPatch(ULandscapePatchComponent* Patch) const
 {
-	return PatchComponents.Contains(Patch);
+	return PatchToIndex.Contains(TSoftObjectPtr<ULandscapePatchComponent>(Patch));
 }
 
 void ALandscapePatchManager::AddPatch(ULandscapePatchComponent* Patch)
@@ -261,7 +279,9 @@ void ALandscapePatchManager::AddPatch(ULandscapePatchComponent* Patch)
 		if (!ContainsPatch(Patch))
 		{
 			Modify();
-			PatchComponents.Add(TSoftObjectPtr<ULandscapePatchComponent>(Patch));
+			TSoftObjectPtr<ULandscapePatchComponent> PatchSoftPtr(Patch);
+			PatchComponents.Add(PatchSoftPtr);
+			PatchToIndex.Add(PatchSoftPtr, PatchComponents.Num()-1);
 		}
 
 		if (Patch->GetPatchManager() != this)
@@ -287,7 +307,20 @@ bool ALandscapePatchManager::RemovePatch(ULandscapePatchComponent* Patch)
 	if (Patch && ContainsPatch(Patch))
 	{
 		Modify();
-		bRemoved = PatchComponents.Remove(TSoftObjectPtr<ULandscapePatchComponent>(Patch)) > 0;
+		TSoftObjectPtr<ULandscapePatchComponent> PatchSoftPtr(Patch);
+		bRemoved = PatchComponents.Remove(PatchSoftPtr) > 0;
+		if (bRemoved)
+		{
+			int32 RemovedIndex = -1;
+			if (PatchToIndex.RemoveAndCopyValue(PatchSoftPtr, RemovedIndex))
+			{
+				for (int32 i = RemovedIndex; i < PatchComponents.Num(); ++i)
+				{
+					PatchToIndex.Add(PatchComponents[i], i);
+				}
+			}
+			
+		}
 
 		// No need to update if the patch was already disabled.Important to avoid needlessly updating while dragging 
 		// a blueprint with a disabled patch (since construction scripts constantly add and remove).
@@ -296,18 +329,28 @@ bool ALandscapePatchManager::RemovePatch(ULandscapePatchComponent* Patch)
 			RequestLandscapeUpdate();
 		}
 	}
-	
+
 	return bRemoved;
 }
 
 int32 ALandscapePatchManager::GetIndexOfPatch(const ULandscapePatchComponent* Patch) const
 {
-	return PatchComponents.IndexOfByKey(Patch);
+	if (const int32* Index = PatchToIndex.Find(TSoftObjectPtr<ULandscapePatchComponent>(Patch)))
+	{
+		return *Index;
+	}
+	return INDEX_NONE;
 }
 
 void ALandscapePatchManager::MovePatchToIndex(ULandscapePatchComponent* Patch, int32 Index)
 {
-	if (!Patch || Index < 0 || GetIndexOfPatch(Patch) == Index)
+	if (!Patch || Index < 0)
+	{
+		return;
+	}
+
+	int32 OriginalIndex = GetIndexOfPatch(Patch);
+	if (OriginalIndex == Index)
 	{
 		return;
 	}
@@ -322,6 +365,18 @@ void ALandscapePatchManager::MovePatchToIndex(ULandscapePatchComponent* Patch, i
 	Index = FMath::Clamp(Index, 0, PatchComponents.Num());
 	PatchComponents.Insert(TSoftObjectPtr<ULandscapePatchComponent>(Patch), Index);
 
+	// Update our index lookup structure
+	int32 OtherEndOfChangedIndices = OriginalIndex < 0 ? PatchComponents.Num() - 1
+		: FMath::Min(OriginalIndex, PatchComponents.Num() - 1);
+
+	int32 StartIndex = FMath::Min(Index, OtherEndOfChangedIndices);
+	int32 EndIndex = FMath::Max(Index, OtherEndOfChangedIndices);
+
+	for (int32 i = StartIndex; i <= EndIndex; ++i)
+	{
+		PatchToIndex.Add(PatchComponents[i], i);
+	}
+
 	if (Patch->IsEnabled())
 	{
 		RequestLandscapeUpdate();
@@ -438,6 +493,17 @@ void ALandscapePatchManager::PostEditChangeProperty(FPropertyChangedEvent& Prope
 	}
 }
 
+void ALandscapePatchManager::PostLoad()
+{
+	Super::PostLoad();
+
+	PatchToIndex.Reset();
+	for (int32 i = 0; i < PatchComponents.Num(); ++i)
+	{
+		PatchToIndex.Add(PatchComponents[i], i);
+	}
+}
+
 void ALandscapePatchManager::CheckForErrors()
 {
 	using namespace LandscapePatchManagerLocals;
diff --git a/Engine/Plugins/Experimental/LandscapePatch/Source/LandscapePatch/Public/LandscapePatchComponent.h b/Engine/Plugins/Experimental/LandscapePatch/Source/LandscapePatch/Public/LandscapePatchComponent.h
--- a/Engine/Plugins/Experimental/LandscapePatch/Source/LandscapePatch/Public/LandscapePatchComponent.h
+++ b/Engine/Plugins/Experimental/LandscapePatch/Source/LandscapePatch/Public/LandscapePatchComponent.h
@@ -190,8 +190,6 @@ struct FLandscapePatchComponentInstanceData : public FSceneComponentInstanceData
 
 	UPROPERTY()
 	TSoftObjectPtr<ALandscapePatchManager> PatchManager = nullptr;
-	UPROPERTY()
-	int32 IndexInManager = -1;
 	
 	UPROPERTY()
 	bool bDirtiedByConstructionScript = false;
diff --git a/Engine/Plugins/Experimental/LandscapePatch/Source/LandscapePatch/Public/LandscapePatchManager.h b/Engine/Plugins/Experimental/LandscapePatch/Source/LandscapePatch/Public/LandscapePatchManager.h
--- a/Engine/Plugins/Experimental/LandscapePatch/Source/LandscapePatch/Public/LandscapePatchManager.h
+++ b/Engine/Plugins/Experimental/LandscapePatch/Source/LandscapePatch/Public/LandscapePatchManager.h
@@ -101,13 +101,14 @@ public:
 	virtual void PostEditUndo() override;
 	virtual void PreSave(FObjectPreSaveContext SaveContext) override;
 	virtual void PostEditChangeProperty(FPropertyChangedEvent& PropertyChangedEvent) override;
+	virtual void PostLoad() override;
 #endif
 	virtual bool IsEditorOnly() const override { return true; }
 	virtual bool NeedsLoadForClient() const override { return false; }
 	virtual bool NeedsLoadForServer() const override { return false; }
 
 protected:
-	
+
 	UPROPERTY()
 	TArray<TSoftObjectPtr<ULandscapePatchComponent>> PatchComponents;
 
@@ -131,6 +132,10 @@ private:
 	// so we want to detect this case.
 	bool bDirtiedByConstructionScript = false;
 #endif
+
+	// Transient table to speed up Contains and IndexOf queries, which are very slow for an array of TSoftObjectPtr's.
+	UPROPERTY(Transient)
+	TMap<TSoftObjectPtr<ULandscapePatchComponent>, int32> PatchToIndex;
 };
 
 #if UE_ENABLE_INCLUDE_ORDER_DEPRECATED_IN_5_2
