LandscapePatch: Make it possible to set the height and weight patch source modes in the blueprint editor. This will not do automatic texture copies on the archetype or allocate UTexture objects on the archetype. Instead the texture creation is defferred to component instancing.

For weight patches, remove automatic patch creation when adding to array in detail panel, because this causes stomping of instances in the world when editing values on the archetype. Instead, use creation through EditInlineNew. This also prevents blueprint weight patches from causing reference issues during autosave, as it sets the flags correctly.

Also add RF_Transactional to all the patch subobjects so that they properly work with undo.

diff --git a/Engine/Plugins/Experimental/LandscapePatch/Source/LandscapePatch/Private/LandscapePatchComponent.cpp b/Engine/Plugins/Experimental/LandscapePatch/Source/LandscapePatch/Private/LandscapePatchComponent.cpp
--- a/Engine/Plugins/Experimental/LandscapePatch/Source/LandscapePatch/Private/LandscapePatchComponent.cpp
+++ b/Engine/Plugins/Experimental/LandscapePatch/Source/LandscapePatch/Private/LandscapePatchComponent.cpp
@@ -472,12 +472,19 @@ FLandscapePatchComponentInstanceData::FLandscapePatchComponentInstanceData(const
 #endif
 }
 
-// Called after rerunning construction scripts (when patch is part of a blueprint) to carry over extra data.
-void ULandscapePatchComponent::ApplyComponentInstanceData(FLandscapePatchComponentInstanceData* ComponentInstanceData)
+// Called before/after rerunning construction scripts (when patch is part of a blueprint) to carry over extra data.
+void ULandscapePatchComponent::ApplyComponentInstanceData(FLandscapePatchComponentInstanceData* ComponentInstanceData, 
+	ECacheApplyPhase CacheApplyPhase)
 {
 #if WITH_EDITOR
 	using namespace LandscapePatchComponentLocals;
 
+	if (CacheApplyPhase == ECacheApplyPhase::PostUserConstructionScript)
+	{
+		// Avoid stomping user construction script changes.
+		return;
+	}
+
 	if (!ComponentInstanceData || !IsRealPatch())
 	{
 		return;
diff --git a/Engine/Plugins/Experimental/LandscapePatch/Source/LandscapePatch/Private/LandscapeTextureBackedRenderTarget.cpp b/Engine/Plugins/Experimental/LandscapePatch/Source/LandscapePatch/Private/LandscapeTextureBackedRenderTarget.cpp
--- a/Engine/Plugins/Experimental/LandscapePatch/Source/LandscapePatch/Private/LandscapeTextureBackedRenderTarget.cpp
+++ b/Engine/Plugins/Experimental/LandscapePatch/Source/LandscapePatch/Private/LandscapeTextureBackedRenderTarget.cpp
@@ -19,6 +19,7 @@ namespace LandscapeTextureBackedRenderTargetLocals
 	UTexture2D* CreateTexture(UObject* Parent)
 	{
 		UTexture2D* Texture = Parent ? NewObject<UTexture2D>(Parent) : NewObject<UTexture2D>();
+		Texture->SetFlags(RF_Transactional);
 		Texture->SRGB = false;
 		Texture->MipGenSettings = TMGS_NoMipmaps;
 		Texture->AddressX = TA_Clamp;
@@ -51,17 +52,25 @@ void ULandscapeTextureBackedRenderTargetBase::PostLoad()
 {
 	Super::PostLoad();
 
+	SetFlags(RF_Transactional);
+	if (IsValid(InternalTexture))
+	{
+		InternalTexture->SetFlags(RF_Transactional);
+	}
+
 	if (!bUseInternalTextureOnly && IsValid(InternalTexture))
 	{
 		InternalTexture->ConditionalPostLoad();
 
-		RenderTarget = NewObject<UTextureRenderTarget2D>(this);
-		RenderTarget->RenderTargetFormat = GetRenderTargetFormat();
-		RenderTarget->InitAutoFormat(SizeX, SizeY);
-		RenderTarget->UpdateResourceImmediate(false);
+		ReinitializeRenderTarget();
 
 		CopyBackFromInternalTexture();
 	}
+
+	if (IsValid(RenderTarget))
+	{
+		RenderTarget->SetFlags(RF_Transactional);
+	}
 }
 
 // PreDuplicate is called not just when objects are copied in editor, but also when they are serialized for undo/redo.
@@ -93,10 +102,7 @@ void ULandscapeTextureBackedRenderTargetBase::PostEditImport()
 
 	if (!bUseInternalTextureOnly && IsValid(InternalTexture))
 	{
-		RenderTarget = NewObject<UTextureRenderTarget2D>(this);
-		RenderTarget->RenderTargetFormat = GetRenderTargetFormat();
-		RenderTarget->InitAutoFormat(SizeX, SizeY);
-		RenderTarget->UpdateResourceImmediate(false);
+		ReinitializeRenderTarget();
 
 		CopyBackFromInternalTexture();
 	}
@@ -130,10 +136,7 @@ void ULandscapeTextureBackedRenderTargetBase::SetUseInternalTextureOnly(bool bUs
 
 		if (IsValid(InternalTexture)) // if initialized
 		{
-			RenderTarget = NewObject<UTextureRenderTarget2D>(this);
-			RenderTarget->RenderTargetFormat = GetRenderTargetFormat();
-			RenderTarget->InitAutoFormat(SizeX, SizeY);
-			RenderTarget->UpdateResourceImmediate(false);
+			ReinitializeRenderTarget();
 
 			if (bCopyExisting)
 			{
@@ -189,14 +192,7 @@ void ULandscapeTextureBackedRenderTargetBase::Initialize()
 	}
 	else
 	{
-		if (!IsValid(RenderTarget))
-		{
-			RenderTarget = NewObject<UTextureRenderTarget2D>(this);
-		}
-
-		RenderTarget->RenderTargetFormat = GetRenderTargetFormat();
-		RenderTarget->InitAutoFormat(SizeX, SizeY);
-		RenderTarget->UpdateResourceImmediate();
+		ReinitializeRenderTarget();
 	}
 #endif // WITH_EDITOR
 }
@@ -312,6 +308,7 @@ void ULandscapeWeightTextureBackedRenderTarget::CopyBackFromInternalTexture()
 	{
 		Modify();
 		RenderTarget = NewObject<UTextureRenderTarget2D>(this);
+		RenderTarget->SetFlags(RF_Transactional);
 		bCreatedNewRenderTarget = true;
 	}
 
@@ -464,6 +461,7 @@ void ULandscapeHeightTextureBackedRenderTarget::CopyBackFromInternalTexture()
 	{
 		Modify();
 		RenderTarget = NewObject<UTextureRenderTarget2D>(this);
+		RenderTarget->SetFlags(RF_Transactional);
 		bCreatedNewRenderTarget = true;
 	}
 
@@ -525,3 +523,20 @@ bool ULandscapeTextureBackedRenderTargetBase::IsCopyingBackAndForthAllowed()
 		&& (!World || (IsValid(World) && World->WorldType == EWorldType::Editor))
 		&& FApp::CanEverRender();
 }
+
+void ULandscapeTextureBackedRenderTargetBase::ReinitializeRenderTarget()
+{
+	if (!IsValid(RenderTarget))
+	{
+		Modify();
+		RenderTarget = NewObject<UTextureRenderTarget2D>(this);
+		RenderTarget->SetFlags(RF_Transactional);
+	}
+	else
+	{
+		RenderTarget->Modify();
+	}
+	RenderTarget->RenderTargetFormat = GetRenderTargetFormat();
+	RenderTarget->InitAutoFormat(SizeX, SizeY);
+	RenderTarget->UpdateResourceImmediate(false);
+}
diff --git a/Engine/Plugins/Experimental/LandscapePatch/Source/LandscapePatch/Private/LandscapeTexturePatch.cpp b/Engine/Plugins/Experimental/LandscapePatch/Source/LandscapePatch/Private/LandscapeTexturePatch.cpp
--- a/Engine/Plugins/Experimental/LandscapePatch/Source/LandscapePatch/Private/LandscapeTexturePatch.cpp
+++ b/Engine/Plugins/Experimental/LandscapePatch/Source/LandscapePatch/Private/LandscapeTexturePatch.cpp
@@ -2,6 +2,7 @@
 
 #include "LandscapeTexturePatch.h"
 
+#include "Engine/World.h"
 #include "Landscape.h"
 #include "LandscapeDataAccess.h"
 #include "LandscapeLayerInfoObject.h" // VisibilityLayer->LayerName
@@ -113,6 +114,11 @@ UTextureRenderTarget2D* ULandscapeTexturePatch::RenderLayer_Native(const FLandsc
 
 		for (const TObjectPtr<ULandscapeWeightPatchTextureInfo>& WeightPatchEntry : WeightPatches)
 		{
+			if (!IsValid(WeightPatchEntry))
+			{
+				continue;
+			}
+
 			if ((bIsWeightmapTarget && (WeightPatchEntry->WeightmapLayerName == InParameters.WeightmapLayerName)) ||
 				(bIsVisibilityLayerTarget && WeightPatchEntry->bEditVisibilityLayer))
 			{
@@ -247,10 +253,10 @@ UTextureRenderTarget2D* ULandscapeTexturePatch::ApplyToWeightmap(ULandscapeWeigh
 	case ELandscapeTexturePatchSourceMode::None:
 		return InCombinedResult;
 	case ELandscapeTexturePatchSourceMode::InternalTexture:
-		PatchUObject = PatchInfo->InternalData->GetInternalTexture();
+		PatchUObject = PatchInfo->InternalData ? PatchInfo->InternalData->GetInternalTexture() : nullptr;
 		break;
 	case ELandscapeTexturePatchSourceMode::TextureBackedRenderTarget:
-		PatchUObject = PatchInfo->InternalData->GetRenderTarget();
+		PatchUObject = PatchInfo->InternalData ? PatchInfo->InternalData->GetRenderTarget() : nullptr;
 		break;
 	case ELandscapeTexturePatchSourceMode::TextureAsset:
 		if (IsValid(PatchInfo->TextureAsset) && !ensureMsgf(PatchInfo->TextureAsset->VirtualTextureStreaming == 0,
@@ -513,6 +519,7 @@ FLandscapeHeightPatchConvertToNativeParams ULandscapeTexturePatch::GetHeightConv
 
 #endif // WITH_EDITOR
 
+
 void ULandscapeTexturePatch::ReinitializeHeight()
 {
 #if WITH_EDITOR
@@ -584,7 +591,7 @@ void ULandscapeTexturePatch::ReinitializeWeights()
 			}
 
 			TArray<TObjectPtr<ULandscapeWeightPatchTextureInfo>> FoundPatches = WeightPatches.FilterByPredicate(
-				[&WeightmapLayerName](const TObjectPtr<ULandscapeWeightPatchTextureInfo>& PatchInfo) { return PatchInfo->WeightmapLayerName == WeightmapLayerName; });
+				[&WeightmapLayerName](const TObjectPtr<ULandscapeWeightPatchTextureInfo>& PatchInfo) { return PatchInfo && PatchInfo->WeightmapLayerName == WeightmapLayerName; });
 
 			if (FoundPatches.IsEmpty())
 			{
@@ -837,7 +844,7 @@ bool ULandscapeTexturePatch::AffectsWeightmapLayer(const FName& InLayerName) con
 
 	for (const TObjectPtr<ULandscapeWeightPatchTextureInfo>& WeightPatch : WeightPatches)
 	{
-		if (WeightPatch->WeightmapLayerName == InLayerName)
+		if (IsValid(WeightPatch) && WeightPatch->WeightmapLayerName == InLayerName)
 		{
 			return true;
 		}
@@ -852,7 +859,10 @@ bool ULandscapeTexturePatch::AffectsVisibilityLayer() const
 		return false;
 	}
 
-	return Algo::AnyOf(WeightPatches, [](const TObjectPtr<ULandscapeWeightPatchTextureInfo>& InWeightPatch) { return InWeightPatch->bEditVisibilityLayer; });
+	return Algo::AnyOf(WeightPatches, [](const TObjectPtr<ULandscapeWeightPatchTextureInfo>& InWeightPatch) 
+	{ 
+		return IsValid(InWeightPatch) && InWeightPatch->bEditVisibilityLayer; 
+	});
 }
 
 // We override IsEnabled to make the patch not request updates when all the source modes are "none"
@@ -869,6 +879,10 @@ bool ULandscapeTexturePatch::IsEnabled() const
 	}
 	for (const TObjectPtr<ULandscapeWeightPatchTextureInfo>& WeightPatch : WeightPatches)
 	{
+		if (!IsValid(WeightPatch))
+		{
+			continue;
+		}
 		if (WeightPatch->SourceMode != ELandscapeTexturePatchSourceMode::None || WeightPatch->bReinitializeOnNextRender)
 		{
 			return true;
@@ -877,6 +891,36 @@ bool ULandscapeTexturePatch::IsEnabled() const
 	// If we got to here, we are enabled, but all of the contained data had a source mode of "None"
 	return false;
 }
+
+void ULandscapeTexturePatch::MakeSureInternalDataIsAllocated()
+{
+	if (!ensure(!IsTemplate()))
+	{
+		return;
+	}
+
+	// If we're instancing a blueprint actor, then it might be missing the needed internal data if the source mode
+	// was set to internal texture or render target, because we can't reliably create those in the archetype.
+	// So, we initialize that data if it's missing but needed when instancing.
+	if (!HeightInternalData && (HeightSourceMode == ELandscapeTexturePatchSourceMode::InternalTexture
+		|| HeightSourceMode == ELandscapeTexturePatchSourceMode::TextureBackedRenderTarget))
+	{
+		TransitionHeightSourceModeInternal(ELandscapeTexturePatchSourceMode::None, HeightSourceMode);
+	}
+
+	for (const TObjectPtr<ULandscapeWeightPatchTextureInfo>& WeightPatch : WeightPatches)
+	{
+		if (!IsValid(WeightPatch))
+		{
+			continue;
+		}
+		if (!WeightPatch->InternalData && (WeightPatch->SourceMode == ELandscapeTexturePatchSourceMode::InternalTexture
+			|| WeightPatch->SourceMode == ELandscapeTexturePatchSourceMode::TextureBackedRenderTarget))
+		{
+			WeightPatch->TransitionSourceModeInternal(ELandscapeTexturePatchSourceMode::None, WeightPatch->SourceMode);
+		}
+	}
+}
 #endif // WITH_EDITOR
 
 void ULandscapeTexturePatch::SnapToLandscape()
@@ -1008,7 +1052,7 @@ void ULandscapeTexturePatch::SetResolution(FVector2D ResolutionIn)
 	ResizePatch(HeightSourceMode, HeightInternalData);
 	for (const TObjectPtr<ULandscapeWeightPatchTextureInfo>& WeightPatch : WeightPatches)
 	{
-		if (ensure(IsValid(WeightPatch)))
+		if (IsValid(WeightPatch))
 		{
 			ResizePatch(WeightPatch->SourceMode, WeightPatch->InternalData);
 		}
@@ -1088,19 +1132,16 @@ void ULandscapeTexturePatch::PostEditChangeProperty(FPropertyChangedEvent& Prope
 		}
 		else if (PropertyChangedEvent.Property->GetFName() == GET_MEMBER_NAME_CHECKED(ULandscapeTexturePatch, WeightPatches))
 		{
-			if (NumWeightPatches != WeightPatches.Num())
+			// In certain cases, changes to the internals of a weight info object trigger a PostEditChangeProperty
+			//  on the patch instead of the info object. For instance this happens when editing the objects in the
+			//  blueprint editor and propagating the change to an instance (something that frequently does not work
+			//  due to propagation being unreliable for this array, see comment on WeightPatches).
+			for (TObjectPtr<ULandscapeWeightPatchTextureInfo>& WeightPatch : WeightPatches)
 			{
-				// User must have added/removed a weight patch, so make sure that all the owner pointers are set
-				for (TObjectPtr<ULandscapeWeightPatchTextureInfo>& WeightPatch : WeightPatches)
+				if (IsValid(WeightPatch))
 				{
-					if (!IsValid(WeightPatch))
-					{
-						// The entries that users add start out as null
-						WeightPatch = NewObject<ULandscapeWeightPatchTextureInfo>(this);
-					}
-					WeightPatch->OwningPatch = this;
+					WeightPatch->SetSourceMode(WeightPatch->DetailPanelSourceMode);
 				}
-				NumWeightPatches = WeightPatches.Num();
 			}
 		}
 		else if (PropertyChangedEvent.Property->GetFName() == GET_MEMBER_NAME_CHECKED(FLandscapeTexturePatchEncodingSettings, ZeroInEncoding)
@@ -1147,7 +1188,6 @@ void ULandscapeWeightPatchTextureInfo::PreDuplicate(FObjectDuplicationParameters
 void ULandscapeWeightPatchTextureInfo::SetSourceMode(ELandscapeTexturePatchSourceMode NewMode)
 {
 #if WITH_EDITOR
-	using namespace LandscapeTexturePatchLocals;
 
 	if (SourceMode == NewMode)
 	{
@@ -1155,8 +1195,29 @@ void ULandscapeWeightPatchTextureInfo::SetSourceMode(ELandscapeTexturePatchSourc
 	}
 	Modify();
 
+	if (!IsTemplate())
+	{
+		TransitionSourceModeInternal(SourceMode, NewMode);
+	}
+	// In a template, it is not safe to try to allocate a texture, etc. All we do is clear out the
+	// texture asset pointer if it is not needed, to avoid referencing assets unnecessarily.
+	else if (SourceMode != ELandscapeTexturePatchSourceMode::TextureAsset)
+	{
+		TextureAsset = nullptr;
+	}
+
+	SourceMode = NewMode;
+	DetailPanelSourceMode = NewMode;
+#endif // WITH_EDITOR
+}
+
+#if WITH_EDITOR
+void ULandscapeWeightPatchTextureInfo::TransitionSourceModeInternal(ELandscapeTexturePatchSourceMode OldMode, ELandscapeTexturePatchSourceMode NewMode)
+{
+	using namespace LandscapeTexturePatchLocals;
+
 	FVector2D Resolution(1, 1);
-	if (OwningPatch.IsValid())
+	if (ULandscapeTexturePatch* OwningPatch = Cast<ULandscapeTexturePatch>(GetOuter()))
 	{
 		Resolution = OwningPatch->GetResolution();
 	}
@@ -1164,14 +1225,12 @@ void ULandscapeWeightPatchTextureInfo::SetSourceMode(ELandscapeTexturePatchSourc
 	TransitionSourceMode<ULandscapeWeightTextureBackedRenderTarget>(SourceMode, NewMode, TextureAsset, InternalData, [&Resolution, this]()
 	{
 		ULandscapeWeightTextureBackedRenderTarget* InternalDataToReturn = NewObject<ULandscapeWeightTextureBackedRenderTarget>(this);
+		InternalDataToReturn->SetFlags(RF_Transactional);
 		InternalDataToReturn->SetResolution(Resolution.X, Resolution.Y);
 		return InternalDataToReturn;
 	});
-
-	SourceMode = NewMode;
-	DetailPanelSourceMode = NewMode;
-#endif // WITH_EDITOR
 }
+#endif
 
 void ULandscapeTexturePatch::SetHeightSourceMode(ELandscapeTexturePatchSourceMode NewMode)
 {
@@ -1184,20 +1243,39 @@ void ULandscapeTexturePatch::SetHeightSourceMode(ELandscapeTexturePatchSourceMod
 	}
 	Modify();
 
+	if (!IsTemplate())
+	{
+		TransitionHeightSourceModeInternal(HeightSourceMode, NewMode);
+	}
+	// In a template, it is not safe to try to allocate a texture, etc. All we do is clear out the
+	// texture asset pointer if it is not needed, to avoid referencing assets unnecessarily.
+	else if (HeightSourceMode != ELandscapeTexturePatchSourceMode::TextureAsset)
+	{
+		HeightTextureAsset = nullptr;
+	}
+
+	HeightSourceMode = NewMode;
+	DetailPanelHeightSourceMode = NewMode;
+#endif // WITH_EDITOR
+}
+
+#if WITH_EDITOR
+void ULandscapeTexturePatch::TransitionHeightSourceModeInternal(ELandscapeTexturePatchSourceMode OldMode, ELandscapeTexturePatchSourceMode NewMode)
+{
+	using namespace LandscapeTexturePatchLocals;
+
 	TransitionSourceMode<ULandscapeHeightTextureBackedRenderTarget>(HeightSourceMode, NewMode, HeightTextureAsset, HeightInternalData, [this]()
 	{
 		ULandscapeHeightTextureBackedRenderTarget* InternalDataToReturn = NewObject<ULandscapeHeightTextureBackedRenderTarget>(this);
+		InternalDataToReturn->SetFlags(RF_Transactional);
 		InternalDataToReturn->SetResolution(ResolutionX, ResolutionY);
 		InternalDataToReturn->SetFormat(HeightRenderTargetFormat);
 		InternalDataToReturn->ConversionParams = GetHeightConvertToNativeParams();
 
 		return InternalDataToReturn;
 	});
-
-	HeightSourceMode = NewMode;
-	DetailPanelHeightSourceMode = NewMode;
-#endif // WITH_EDITOR
 }
+#endif
 
 void ULandscapeTexturePatch::SetHeightTextureAsset(UTexture* TextureIn)
 {
@@ -1292,6 +1370,11 @@ void ULandscapeTexturePatch::AddWeightPatch(const FName& WeightmapLayerName, ELa
 	// Try to modify an existing entry instead if possible
 	for (const TObjectPtr<ULandscapeWeightPatchTextureInfo>& WeightPatch : WeightPatches)
 	{
+		if (!IsValid(WeightPatch))
+		{
+			continue;
+		}
+
 		if (WeightPatch->WeightmapLayerName == WeightmapLayerName)
 		{
 			if (WeightPatch->SourceMode != SourceMode)
@@ -1306,46 +1389,55 @@ void ULandscapeTexturePatch::AddWeightPatch(const FName& WeightmapLayerName, ELa
 		}
 	}
 
-	ULandscapeWeightPatchTextureInfo* NewWeightPatch = NewObject<ULandscapeWeightPatchTextureInfo>(this);
+	// The object creation is modeled after SPropertyEditorEditInline::OnClassPicked, which is how these are created
+	// from the detail panel. We probably don't need the archetype check, admittedly, but might as well keep it.
+	EObjectFlags NewObjectFlags = GetMaskedFlags(RF_PropagateToSubObjects);
+	if (HasAnyFlags(RF_ClassDefaultObject | RF_ArchetypeObject))
+	{
+		NewObjectFlags |= RF_ArchetypeObject;
+	}
+	ULandscapeWeightPatchTextureInfo* NewWeightPatch = NewObject<ULandscapeWeightPatchTextureInfo>(this, NAME_None, NewObjectFlags);
+
 	NewWeightPatch->WeightmapLayerName = WeightmapLayerName;
 	NewWeightPatch->SourceMode = SourceMode;
 	NewWeightPatch->DetailPanelSourceMode = SourceMode;
 	NewWeightPatch->bUseAlphaChannel = bUseAlphaChannel;
-	NewWeightPatch->OwningPatch = this;
 
 	if (NewWeightPatch->SourceMode == ELandscapeTexturePatchSourceMode::InternalTexture
 		|| NewWeightPatch->SourceMode == ELandscapeTexturePatchSourceMode::TextureBackedRenderTarget)
 	{
 		NewWeightPatch->InternalData = NewObject<ULandscapeWeightTextureBackedRenderTarget>(NewWeightPatch);
+		NewWeightPatch->InternalData->SetFlags(RF_Transactional);
 		NewWeightPatch->InternalData->SetResolution(ResolutionX, ResolutionY);
 		NewWeightPatch->InternalData->SetUseAlphaChannel(bUseAlphaChannel);
 		NewWeightPatch->InternalData->Initialize();
 	}
 
 	WeightPatches.Add(NewWeightPatch);
-	++NumWeightPatches;
 #endif // WITH_EDITOR
 }
 
 void ULandscapeTexturePatch::RemoveWeightPatch(const FName& InWeightmapLayerName)
 {
-	NumWeightPatches -= WeightPatches.RemoveAll([InWeightmapLayerName](const TObjectPtr<ULandscapeWeightPatchTextureInfo>& PatchInfo)
+	WeightPatches.RemoveAll([InWeightmapLayerName](const TObjectPtr<ULandscapeWeightPatchTextureInfo>& PatchInfo)
 	{ 
-		return PatchInfo->WeightmapLayerName == InWeightmapLayerName; 
+		return PatchInfo && PatchInfo->WeightmapLayerName == InWeightmapLayerName; 
 	});
 }
 
 void ULandscapeTexturePatch::RemoveAllWeightPatches()
 {
 	WeightPatches.Reset();
-	NumWeightPatches = 0;
 }
 
 void ULandscapeTexturePatch::DisableAllWeightPatches()
 {
 	for (const TObjectPtr<ULandscapeWeightPatchTextureInfo>& WeightPatch : WeightPatches)
 	{
-		WeightPatch->SetSourceMode(ELandscapeTexturePatchSourceMode::None);
+		if (IsValid(WeightPatch))
+		{
+			WeightPatch->SetSourceMode(ELandscapeTexturePatchSourceMode::None);
+		}
 	}
 }
 
@@ -1354,7 +1446,7 @@ TArray<FName> ULandscapeTexturePatch::GetAllWeightPatchLayerNames()
 	TArray<FName> Names;
 	for (const TObjectPtr<ULandscapeWeightPatchTextureInfo>& WeightPatch : WeightPatches)
 	{
-		if (WeightPatch->WeightmapLayerName != NAME_None)
+		if (IsValid(WeightPatch) && WeightPatch->WeightmapLayerName != NAME_None)
 		{
 			Names.AddUnique(WeightPatch->WeightmapLayerName);
 		}
@@ -1367,7 +1459,7 @@ void ULandscapeTexturePatch::SetUseAlphaChannelForWeightPatch(const FName& InWei
 {
 	for (const TObjectPtr<ULandscapeWeightPatchTextureInfo>& WeightPatch : WeightPatches)
 	{
-		if (WeightPatch->WeightmapLayerName == InWeightmapLayerName)
+		if (IsValid(WeightPatch) && WeightPatch->WeightmapLayerName == InWeightmapLayerName)
 		{
 			WeightPatch->bUseAlphaChannel = bUseAlphaChannel;
 			if (WeightPatch->InternalData)
@@ -1388,12 +1480,9 @@ void ULandscapeTexturePatch::SetWeightPatchSourceMode(const FName& InWeightmapLa
 
 	for (const TObjectPtr<ULandscapeWeightPatchTextureInfo>& WeightPatch : WeightPatches)
 	{
-		if (WeightPatch->WeightmapLayerName == InWeightmapLayerName)
+		if (IsValid(WeightPatch) && WeightPatch->WeightmapLayerName == InWeightmapLayerName)
 		{
-			if (WeightPatch->SourceMode != NewMode)
-			{
-				WeightPatch->SetSourceMode(NewMode);
-			}
+			WeightPatch->SetSourceMode(NewMode);
 			return;
 		}
 	}
@@ -1406,7 +1495,7 @@ ELandscapeTexturePatchSourceMode ULandscapeTexturePatch::GetWeightPatchSourceMod
 {
 	for (const TObjectPtr<ULandscapeWeightPatchTextureInfo>& WeightPatch : WeightPatches)
 	{
-		if (WeightPatch->WeightmapLayerName == InWeightmapLayerName)
+		if (IsValid(WeightPatch) && WeightPatch->WeightmapLayerName == InWeightmapLayerName)
 		{
 			return WeightPatch->SourceMode;
 		}
@@ -1418,7 +1507,7 @@ UTextureRenderTarget2D* ULandscapeTexturePatch::GetWeightPatchRenderTarget(const
 {
 	for (const TObjectPtr<ULandscapeWeightPatchTextureInfo>& WeightPatch : WeightPatches)
 	{
-		if (WeightPatch->WeightmapLayerName == InWeightmapLayerName)
+		if (IsValid(WeightPatch) && WeightPatch->WeightmapLayerName == InWeightmapLayerName)
 		{
 			return WeightPatch->InternalData ? WeightPatch->InternalData->GetRenderTarget() : nullptr;
 		}
@@ -1430,7 +1519,7 @@ UTexture2D* ULandscapeTexturePatch::GetWeightPatchInternalTexture(const FName& I
 {
 	for (const TObjectPtr<ULandscapeWeightPatchTextureInfo>& WeightPatch : WeightPatches)
 	{
-		if (WeightPatch->WeightmapLayerName == InWeightmapLayerName)
+		if (IsValid(WeightPatch) && WeightPatch->WeightmapLayerName == InWeightmapLayerName)
 		{
 			return WeightPatch->InternalData ? WeightPatch->InternalData->GetInternalTexture() : nullptr;
 		}
@@ -1448,7 +1537,7 @@ void ULandscapeTexturePatch::SetWeightPatchTextureAsset(const FName& InWeightmap
 
 	for (const TObjectPtr<ULandscapeWeightPatchTextureInfo>& WeightPatch : WeightPatches)
 	{
-		if (WeightPatch->WeightmapLayerName == InWeightmapLayerName)
+		if (IsValid(WeightPatch) && WeightPatch->WeightmapLayerName == InWeightmapLayerName)
 		{
 			WeightPatch->TextureAsset = TextureIn;
 			return;
@@ -1463,7 +1552,7 @@ void ULandscapeTexturePatch::SetWeightPatchBlendModeOverride(const FName& InWeig
 {
 	for (const TObjectPtr<ULandscapeWeightPatchTextureInfo>& WeightPatch : WeightPatches)
 	{
-		if (WeightPatch->WeightmapLayerName == InWeightmapLayerName)
+		if (IsValid(WeightPatch) && WeightPatch->WeightmapLayerName == InWeightmapLayerName)
 		{
 			WeightPatch->OverrideBlendMode = BlendModeIn;
 			WeightPatch->bOverrideBlendMode = true;
@@ -1476,7 +1565,7 @@ void ULandscapeTexturePatch::ClearWeightPatchBlendModeOverride(const FName& InWe
 {
 	for (const TObjectPtr<ULandscapeWeightPatchTextureInfo>& WeightPatch : WeightPatches)
 	{
-		if (WeightPatch->WeightmapLayerName == InWeightmapLayerName)
+		if (IsValid(WeightPatch) && WeightPatch->WeightmapLayerName == InWeightmapLayerName)
 		{
 			WeightPatch->bOverrideBlendMode = false;
 			return;
@@ -1488,9 +1577,123 @@ void ULandscapeTexturePatch::SetEditVisibilityLayer(const FName& InWeightmapLaye
 {
 	for (const TObjectPtr<ULandscapeWeightPatchTextureInfo>& WeightPatch : WeightPatches)
 	{
-		if (WeightPatch->WeightmapLayerName == InWeightmapLayerName)
+		if (IsValid(WeightPatch) && WeightPatch->WeightmapLayerName == InWeightmapLayerName)
 		{
 			WeightPatch->bEditVisibilityLayer = bEditVisibilityLayer;
 		}
 	}
 }
+
+#if WITH_EDITOR
+TStructOnScope<FActorComponentInstanceData> ULandscapeTexturePatch::GetComponentInstanceData() const
+{
+	return MakeStructOnScope<FActorComponentInstanceData, FLandscapeTexturePatchInstanceData>(this);
+}
+#endif
+
+FLandscapeTexturePatchInstanceData::FLandscapeTexturePatchInstanceData(const ULandscapeTexturePatch* Patch)
+	: FLandscapePatchComponentInstanceData(Patch)
+{
+#if WITH_EDITOR
+
+	if (!ensure(Patch))
+	{
+		return;
+	}
+
+	HeightInternalData = Patch->HeightInternalData;
+	if (HeightInternalData.IsValid())
+	{
+		// Disconnect owner so that the object doesn't get destroyed
+		HeightInternalData->Rename(nullptr, GetTransientPackage());
+	}
+
+	for (const TObjectPtr<ULandscapeWeightPatchTextureInfo>& WeightPatch : Patch->WeightPatches)
+	{
+		if (!IsValid(WeightPatch))
+		{
+			continue;
+		}
+
+		WeightPatchInternalData.Add(WeightPatch->InternalData);
+		
+		if (IsValid(WeightPatch->InternalData))
+		{
+			// Disconnect owner so that the object doesn't get destroyed
+			WeightPatch->InternalData->Rename(nullptr, GetTransientPackage());
+		}
+	}
+#endif
+}
+
+// Called before/after rerunning construction scripts
+void ULandscapeTexturePatch::ApplyComponentInstanceData(FLandscapeTexturePatchInstanceData* ComponentInstanceData, ECacheApplyPhase CacheApplyPhase)
+{
+#if WITH_EDITOR
+
+	if (CacheApplyPhase == ECacheApplyPhase::PostUserConstructionScript)
+	{
+		// Avoid stomping user construction script changes.
+		return;
+	}
+
+	UWorld* World = GetWorld();
+	if (!ComponentInstanceData 
+		// TODO: Replace with !IsPatchInWorld() once that CL is in
+		|| IsTemplate() || !IsValid(this) || !IsValid(World) || World->WorldType != EWorldType::Editor)
+	{
+		return;
+	}
+
+	if (ComponentInstanceData->HeightInternalData.IsValid() 
+		&& (HeightSourceMode == ELandscapeTexturePatchSourceMode::InternalTexture
+			|| HeightSourceMode == ELandscapeTexturePatchSourceMode::TextureBackedRenderTarget))
+	{
+		HeightInternalData = ComponentInstanceData->HeightInternalData.Get();
+
+		// Attach to the new object
+		HeightInternalData->Rename(nullptr, this);
+	}
+	
+	for (int32 i = 0; i < ComponentInstanceData->WeightPatchInternalData.Num(); ++i)
+	{
+		if (i >= WeightPatches.Num() || !IsValid(WeightPatches[i]) || !ComponentInstanceData->WeightPatchInternalData[i].IsValid())
+		{
+			continue;
+		}
+
+		if (WeightPatches[i]->SourceMode == ELandscapeTexturePatchSourceMode::InternalTexture
+			|| WeightPatches[i]->SourceMode == ELandscapeTexturePatchSourceMode::TextureBackedRenderTarget)
+		{
+			WeightPatches[i]->InternalData = ComponentInstanceData->WeightPatchInternalData[i].Get();
+
+			// Attach to the new object
+			WeightPatches[i]->InternalData->Rename(nullptr, WeightPatches[i]);
+		}
+	}
+
+	// Note: For the below, one might think that we should override PostApplyToComponent, but that is actually
+	//  unreliable because it is called from FActorComponentInstanceData::ApplyToComponent, and gets skipped
+	//  if that subclass didn't happen to serialize some properties itself. Also, if it does get called, it is
+	//  not at the end because the Super call gets done earlier (not just by us, but by intermediate subclasses).
+
+	// If we didn't carry over internal data, we need to make sure at this point that it exists, because
+	//  we cannot create it in the template
+	MakeSureInternalDataIsAllocated();
+
+	// Make sure the blueprint-set resolution and format is applied
+	if (HeightInternalData)
+	{
+		HeightInternalData->SetResolution(ResolutionX, ResolutionY);
+		HeightInternalData->SetFormat(HeightRenderTargetFormat);
+	}
+	for (TObjectPtr<ULandscapeWeightPatchTextureInfo>& WeightPatch : WeightPatches)
+	{
+		if (IsValid(WeightPatch) && WeightPatch->InternalData)
+		{
+			WeightPatch->InternalData->SetResolution(ResolutionX, ResolutionY);
+		}
+	}
+	
+#endif // WITH_EDITOR
+}
diff --git a/Engine/Plugins/Experimental/LandscapePatch/Source/LandscapePatch/Public/LandscapePatchComponent.h b/Engine/Plugins/Experimental/LandscapePatch/Source/LandscapePatch/Public/LandscapePatchComponent.h
--- a/Engine/Plugins/Experimental/LandscapePatch/Source/LandscapePatch/Public/LandscapePatchComponent.h
+++ b/Engine/Plugins/Experimental/LandscapePatch/Source/LandscapePatch/Public/LandscapePatchComponent.h
@@ -67,7 +67,7 @@ public:
 	UFUNCTION(BlueprintCallable, Category = "LandscapePatch")
 	virtual ALandscapePatchManager* GetPatchManager() const;
 
-	void ApplyComponentInstanceData(struct FLandscapePatchComponentInstanceData* ComponentInstanceData);
+	void ApplyComponentInstanceData(struct FLandscapePatchComponentInstanceData* ComponentInstanceData, ECacheApplyPhase CacheApplyPhase);
 
 	// For now we keep the patches largely editor-only, since we don't yet support runtime landscape editing.
 	// The above functions are also editor-only (and don't work at runtime), but can't be in WITH_EDITOR blocks
@@ -181,7 +181,7 @@ struct FLandscapePatchComponentInstanceData : public FSceneComponentInstanceData
 	virtual void ApplyToComponent(UActorComponent* Component, const ECacheApplyPhase CacheApplyPhase) override
 	{
 		Super::ApplyToComponent(Component, CacheApplyPhase);
-		CastChecked<ULandscapePatchComponent>(Component)->ApplyComponentInstanceData(this);
+		CastChecked<ULandscapePatchComponent>(Component)->ApplyComponentInstanceData(this, CacheApplyPhase);
 	}
 
 #if WITH_EDITORONLY_DATA
diff --git a/Engine/Plugins/Experimental/LandscapePatch/Source/LandscapePatch/Public/LandscapeTextureBackedRenderTarget.h b/Engine/Plugins/Experimental/LandscapePatch/Source/LandscapePatch/Public/LandscapeTextureBackedRenderTarget.h
--- a/Engine/Plugins/Experimental/LandscapePatch/Source/LandscapePatch/Public/LandscapeTextureBackedRenderTarget.h
+++ b/Engine/Plugins/Experimental/LandscapePatch/Source/LandscapePatch/Public/LandscapeTextureBackedRenderTarget.h
@@ -60,6 +60,8 @@ protected:
 	virtual ETextureSourceFormat GetInternalTextureFormat() PURE_VIRTUAL(ULandscapeTextureBackedRenderTargetBase::GetInternalTextureFormat, return ETextureSourceFormat::TSF_G8;);
 	virtual ETextureRenderTargetFormat GetRenderTargetFormat() PURE_VIRTUAL(ULandscapeTextureBackedRenderTargetBase::GetRenderTargetFormat, return ETextureRenderTargetFormat::RTF_R8;);
 
+	void ReinitializeRenderTarget();
+
 	UPROPERTY(VisibleAnywhere, Category = InternalData)
 	TObjectPtr<UTexture2D> InternalTexture = nullptr;
 
diff --git a/Engine/Plugins/Experimental/LandscapePatch/Source/LandscapePatch/Public/LandscapeTexturePatch.h b/Engine/Plugins/Experimental/LandscapePatch/Source/LandscapePatch/Public/LandscapeTexturePatch.h
--- a/Engine/Plugins/Experimental/LandscapePatch/Source/LandscapePatch/Public/LandscapeTexturePatch.h
+++ b/Engine/Plugins/Experimental/LandscapePatch/Source/LandscapePatch/Public/LandscapeTexturePatch.h
@@ -142,11 +142,12 @@ public:
 	double WorldSpaceEncodingScale = 1;
 };
 
+//~ Ideally this would be a nested class, but it needs to be a UObject, which can't be nested.
 /**
  * Helper class for ULandscapeTexturePatch that stores information for a given weight layer.
- * Should not be used outside this class, but does need to be a UObject (so can't be nested).
+ * Should not be used outside this class.
  */
-UCLASS()
+UCLASS(EditInlineNew, CollapseCategories)
 class LANDSCAPEPATCH_API ULandscapeWeightPatchTextureInfo : public UObject
 {
 	GENERATED_BODY()
@@ -179,9 +180,10 @@ protected:
 	UPROPERTY(EditAnywhere, Category = WeightPatch, meta = (EditConditionHides, EditCondition = "false"))
 	ELandscapeTexturePatchSourceMode SourceMode = ELandscapeTexturePatchSourceMode::None;
 
-	 //~ This is EditInstanceOnly because we can't create textures in the blueprint editor due to the way
-	 //~ instanced properties are currently handled there.
-	UPROPERTY(EditInstanceOnly, Category = WeightPatch, meta = (DisplayName = "Source Mode"))
+	/**
+	 * How the heightmap of the patch is stored.
+	 */
+	UPROPERTY(EditAnywhere, Category = WeightPatch, meta = (DisplayName = "Source Mode"))
 	ELandscapeTexturePatchSourceMode DetailPanelSourceMode = ELandscapeTexturePatchSourceMode::None;
 
 	//~ We could refactor things such that we always have an InternalData pointer, even when we use
@@ -199,19 +201,20 @@ protected:
 	// TODO: We could support having different per-layer falloff modes and falloff amounts as well, as
 	// additional override members. But probably better to wait to see if that is actually desired.
 
-	// Needed mainly so that we can get at the resolution...
-	UPROPERTY()
-	TWeakObjectPtr<ULandscapeTexturePatch> OwningPatch = nullptr;
-
 	// TODO: Like the similar flag for the height patch, this might not work once local merge works
 	// with landscape brushes...
 	bool bReinitializeOnNextRender = false;
 
 	void SetSourceMode(ELandscapeTexturePatchSourceMode NewMode);
+#if WITH_EDITOR
+	void TransitionSourceModeInternal(ELandscapeTexturePatchSourceMode OldMode, ELandscapeTexturePatchSourceMode NewMode);
+#endif
 
 	friend class ULandscapeTexturePatch;
+	friend struct FLandscapeTexturePatchInstanceData;
 };
 
+
 UCLASS(Blueprintable, BlueprintType, ClassGroup = Landscape, meta = (BlueprintSpawnableComponent))
 class LANDSCAPEPATCH_API ULandscapeTexturePatch : public ULandscapePatchComponent
 {
@@ -227,6 +230,9 @@ public:
 	virtual bool AffectsVisibilityLayer() const override;
 	virtual bool IsEnabled() const override;
 
+	// UActorComponent
+	virtual TStructOnScope<FActorComponentInstanceData> GetComponentInstanceData() const override;
+
 	// UObject
 	virtual void PostEditChangeProperty(FPropertyChangedEvent& PropertyChangedEvent) override;
 #endif
@@ -451,10 +457,12 @@ public:
 	UTexture2D* GetWeightPatchInternalTexture(const FName& InWeightmapLayerName);
 
 protected:
-
-	UPROPERTY()
+	//~ Don't expose these on the instance because a user might not realize that they would lose their existing internal
+	//~ data by dragging them, and the only way they can reinitialize data in the viewport is through the methods that
+	//~ already use InitTextureSizeX/Y as inputs
+	UPROPERTY(EditDefaultsOnly, Category = Settings)
 	int32 ResolutionX = 32;
-	UPROPERTY()
+	UPROPERTY(EditDefaultsOnly, Category = Settings)
 	int32 ResolutionY = 32;
 
 	/** At scale 1.0, the X and Y of the region affected by the height patch. This corresponds to the distance from the center
@@ -486,13 +494,9 @@ protected:
 	ELandscapeTexturePatchSourceMode HeightSourceMode = ELandscapeTexturePatchSourceMode::None;
 
 	/**
-	 * How the heightmap of the patch is stored. Not settable in the detail panel of the blueprint editor- use SetHeightSourceMode
-	 * in blueprint actors instead.
+	 * How the heightmap of the patch is stored.
 	 */
-	//~ This is EditInstanceOnly because changing it creates/destroys the internal texture, and that cannot currently be
-	//~ dealt with properly in the blueprint editor due to the way instanced properties are handled there. Could revisit when
-	//~ UE-158706 is resolved.
-	UPROPERTY(EditInstanceOnly, Category = HeightPatch, meta = (DisplayName = "Source Mode"))
+	UPROPERTY(EditAnywhere, Category = HeightPatch, meta = (DisplayName = "Source Mode"))
 	ELandscapeTexturePatchSourceMode DetailPanelHeightSourceMode = ELandscapeTexturePatchSourceMode::None;
 
 	/** Not directly settable via detail panel- for display/debugging purposes only. */
@@ -541,20 +545,16 @@ protected:
 	// Weight properties:
 
 	/** 
-	 * Weight patches. These are not available to be manipulated through the detail panel in the blueprint editor, so blueprint actors
-	 * should set up the weight patches via AddWeightPatch instead.
-	 */
-	//~ This is EditInstanceOnly because manipulating them in blueprint editor causes saving issues due to the way that
-	//~ instanced properties are currently handled there.
-	UPROPERTY(EditInstanceOnly, Category = WeightPatches, Instanced)
+	 * Weight patches. 
+	 * Note that manipulating these in the blueprint editor will not reliably update instances that are already
+	 * placed into the world, due to current limitations in how change detection is done for such arrays. Specifically,
+	 * existing instances that are actually not customized are very likely to be erroneously be treated as having
+	 * customized their version of the array, causing the blueprint changes to not be pushed to those instances
+	 * when they otherwise would be for most other properties.
+	 */
+	UPROPERTY(EditAnywhere, Category = WeightPatches, Instanced, NoClear, meta=(NoResetToDefault))
 	TArray<TObjectPtr<ULandscapeWeightPatchTextureInfo>> WeightPatches;
 
-	// Used to detect changes to the number of weight patches via the detail panel, so that we can
-	// initialize the "owner" pointer when patches are added or trigger update when patches are removed.
-	UPROPERTY()
-	int32 NumWeightPatches = 0;
-
-
 	// Reinitialization from detail panel:
 
 	/**
@@ -593,19 +593,23 @@ protected:
 		EditCondition = "bBaseResolutionOffLandscape"))
 	float ResolutionMultiplier = 1;
 
-	/** Texture width to use when reinitializing, if not basing resolution off landscape. */
+	/** Texture width to use when reinitializing using Reinitialize Weights or ReinitializeHeight, if not basing resolution off landscape. */
 	UPROPERTY(EditAnywhere, Category = Initialization, meta = (EditCondition = "!bBaseResolutionOffLandscape", 
 		ClampMin = "1"))
 	int32 InitTextureSizeX = 33;
 
-	/** Texture height to use when reinitializing */
+	/** Texture height to use when reinitializing using Reinitialize Weights or ReinitializeHeight, if not basing resolution off landscape. */
 	UPROPERTY(EditAnywhere, Category = Initialization, meta = (EditCondition = "!bBaseResolutionOffLandscape", 
 		ClampMin = "1"))
 	int32 InitTextureSizeY = 33;
 
 private:
+	void ApplyComponentInstanceData(struct FLandscapeTexturePatchInstanceData* ComponentInstanceData, ECacheApplyPhase CacheApplyPhase);
+	friend struct FLandscapeTexturePatchInstanceData;
+
 	void UpdateHeightConvertToNativeParamsIfNeeded();
 #if WITH_EDITOR
+	void TransitionHeightSourceModeInternal(ELandscapeTexturePatchSourceMode OldMode, ELandscapeTexturePatchSourceMode NewMode);
 	FLandscapeHeightPatchConvertToNativeParams GetHeightConvertToNativeParams() const;
 	UTextureRenderTarget2D* ApplyToHeightmap(UTextureRenderTarget2D* InCombinedResult);
 	UTextureRenderTarget2D* ApplyToWeightmap(ULandscapeWeightPatchTextureInfo* PatchInfo, UTextureRenderTarget2D* InCombinedResult);
@@ -621,12 +625,45 @@ private:
 	FMatrix44f GetPatchToHeightmapUVs(int32 PatchSizeX, int32 PatchSizeY, int32 HeightmapSizeX, int32 HeightmapSizeY) const;
 	void ReinitializeHeight(UTextureRenderTarget2D* InCombinedResult);
 	void ReinitializeWeightPatch(ULandscapeWeightPatchTextureInfo* PatchInfo, UTextureRenderTarget2D* InCombinedResult);
+
+	void MakeSureInternalDataIsAllocated();
 #endif // WITH_EDITOR
 
-	UPROPERTY()
+	UPROPERTY(EditDefaultsOnly, Category = Settings)
 	TEnumAsByte<ETextureRenderTargetFormat> HeightRenderTargetFormat = ETextureRenderTargetFormat::RTF_R32f;
 };
 
+
+/** Carries over internal data subobjects to avoid constant reconstruction on construction script reruns. */
+USTRUCT()
+struct FLandscapeTexturePatchInstanceData : public FLandscapePatchComponentInstanceData
+{
+	GENERATED_BODY()
+
+	FLandscapeTexturePatchInstanceData() = default;
+	FLandscapeTexturePatchInstanceData(const ULandscapeTexturePatch* SourceComponent);
+
+	virtual ~FLandscapeTexturePatchInstanceData() = default;
+
+	virtual bool ContainsData() const override
+	{
+		return true;
+	}
+
+	virtual void ApplyToComponent(UActorComponent* Component, const ECacheApplyPhase CacheApplyPhase) override
+	{
+		Super::ApplyToComponent(Component, CacheApplyPhase);
+		CastChecked<ULandscapeTexturePatch>(Component)->ApplyComponentInstanceData(this, CacheApplyPhase);
+	}
+
+#if WITH_EDITORONLY_DATA
+	UPROPERTY()
+	TWeakObjectPtr<ULandscapeHeightTextureBackedRenderTarget> HeightInternalData = nullptr;
+	UPROPERTY()
+	TArray<TWeakObjectPtr<ULandscapeWeightTextureBackedRenderTarget>> WeightPatchInternalData;
+#endif
+};
+
 #if UE_ENABLE_INCLUDE_ORDER_DEPRECATED_IN_5_2
 #include "CoreMinimal.h"
 #endif
