Fixed crash on Metal in LandscapeCircleHeightPatchPS when attempting to write 2 components to a 4 component RT

diff --git a/Engine/Plugins/Experimental/LandscapePatch/Shaders/Private/LandscapeCircleHeightPatchPS.usf b/Engine/Plugins/Experimental/LandscapePatch/Shaders/Private/LandscapeCircleHeightPatchPS.usf
--- a/Engine/Plugins/Experimental/LandscapePatch/Shaders/Private/LandscapeCircleHeightPatchPS.usf
+++ b/Engine/Plugins/Experimental/LandscapePatch/Shaders/Private/LandscapeCircleHeightPatchPS.usf
@@ -12,7 +12,7 @@ Texture2D<float4> InSourceTexture;
 
 #if CIRCLE_HEIGHT_PATCH
 
-void ApplyLandscapeCircleHeightPatch(in float4 SVPos : SV_POSITION, out float2 OutColor : SV_Target0)
+void ApplyLandscapeCircleHeightPatch(in float4 SVPos : SV_POSITION, out float4 OutColor : SV_Target0)
 {
 	const float KINDA_SMALL_NUMBER = 0.0001;
 	const float ClampedFalloff = max(InFalloff, KINDA_SMALL_NUMBER);
@@ -33,7 +33,7 @@ void ApplyLandscapeCircleHeightPatch(in float4 SVPos : SV_POSITION, out float2 O
 	
 	float NewHeight = lerp(CurrentHeight, PatchHeight, Alpha);
 	
-	OutColor = PackHeight(NewHeight);
+	OutColor = float4(PackHeight(NewHeight), 0.0, 0.0);
 }
 
 #endif // CIRCLE_HEIGHT_PATCH
@@ -57,4 +57,4 @@ void ApplyLandscapeCircleVisibilityPatch(in float4 SVPos : SV_POSITION, out floa
 	}
 }
 
-#endif // CIRCLE_VISIBILITY_PATCH
\ No newline at end of file
+#endif // CIRCLE_VISIBILITY_PATCH
