LandscapePatch: Allow patches be registered with a custom edit layer directly (instead of a patch manager) and order themselves using a priority value. Edit layers will eventually be renderer providers in the upcoming batched merge landscape path, but for now, add a legacy support interface to allow them to receive brush-like render calls to render their patches in global merge.

diff --git a/Engine/Plugins/Experimental/LandscapePatch/Source/LandscapePatch/LandscapePatch.Build.cs b/Engine/Plugins/Experimental/LandscapePatch/Source/LandscapePatch/LandscapePatch.Build.cs
--- a/Engine/Plugins/Experimental/LandscapePatch/Source/LandscapePatch/LandscapePatch.Build.cs
+++ b/Engine/Plugins/Experimental/LandscapePatch/Source/LandscapePatch/LandscapePatch.Build.cs
@@ -41,5 +41,18 @@ public class LandscapePatch : ModuleRules
 				"Projects", // IPluginManager
 			}
 			);
+			
+		if (Target.bBuildEditor)
+        {
+			PrivateDependencyModuleNames.AddRange(
+				new string[] {
+					// These are for some fix-up operations to transition patches to a newer
+					// system and delete some old actors (triggered by user).
+					"LevelEditor",
+					"TypedElementRuntime",
+					"UnrealEd", // FScopedTransaction
+				}
+			);
+		}
 	}
 }
diff --git a/Engine/Plugins/Experimental/LandscapePatch/Source/LandscapePatch/Private/LandscapeCircleHeightPatch.cpp b/Engine/Plugins/Experimental/LandscapePatch/Source/LandscapePatch/Private/LandscapeCircleHeightPatch.cpp
--- a/Engine/Plugins/Experimental/LandscapePatch/Source/LandscapePatch/Private/LandscapeCircleHeightPatch.cpp
+++ b/Engine/Plugins/Experimental/LandscapePatch/Source/LandscapePatch/Private/LandscapeCircleHeightPatch.cpp
@@ -40,26 +40,14 @@ void ULandscapeCircleHeightPatch::OnComponentCreated()
 	}
 }
 
-void ULandscapeCircleHeightPatch::Initialize_Native(const FTransform & InLandscapeTransform,
-	const FIntPoint& InLandscapeSize,
-	const FIntPoint& InLandscapeRenderTargetSize)
+UTextureRenderTarget2D* ULandscapeCircleHeightPatch::RenderLayer_Native(const FLandscapeBrushParameters& InParameters, const FTransform& HeightmapCoordsToWorld)
 {
-}
-
-UTextureRenderTarget2D* ULandscapeCircleHeightPatch::RenderLayer_Native(const FLandscapeBrushParameters& InParameters)
-{
-	if (!ensure(PatchManager.IsValid()))
-	{
-		return InParameters.CombinedResult;
-	}
-
 	// Circle height patch doesn't affect regular weightmap layers.
 	if ((bEditVisibility && (InParameters.LayerType != ELandscapeToolTargetType::Visibility)) || (!bEditVisibility && (InParameters.LayerType != ELandscapeToolTargetType::Heightmap)))
 	{
 		return InParameters.CombinedResult;
 	}
 
-	FTransform HeightmapCoordsToWorld = PatchManager->GetHeightmapCoordsToWorld();
 	double ToHeightmapRadiusScale = GetComponentTransform().GetScale3D().X / HeightmapCoordsToWorld.GetScale3D().X;
 	FVector3d CircleCenterWorld = GetComponentTransform().GetTranslation();
 	FVector3d CenterInHeightmapCoordinates = HeightmapCoordsToWorld.InverseTransformPosition(CircleCenterWorld);
diff --git a/Engine/Plugins/Experimental/LandscapePatch/Source/LandscapePatch/Private/LandscapePatchComponent.cpp b/Engine/Plugins/Experimental/LandscapePatch/Source/LandscapePatch/Private/LandscapePatchComponent.cpp
--- a/Engine/Plugins/Experimental/LandscapePatch/Source/LandscapePatch/Private/LandscapePatchComponent.cpp
+++ b/Engine/Plugins/Experimental/LandscapePatch/Source/LandscapePatch/Private/LandscapePatchComponent.cpp
@@ -6,81 +6,83 @@
 #include "EngineUtils.h"
 #include "HAL/IConsoleManager.h" // FAutoConsoleVariableRef
 #include "Landscape.h"
+#include "LandscapeModule.h"
 #include "LandscapePatchLogging.h"
 #include "LandscapePatchManager.h"
+#include "LandscapePatchUtil.h" // GetHeightmapToWorld
 #include "Logging/MessageLog.h"
 #include "Misc/UObjectToken.h"
+#include "Modules/ModuleManager.h"
 #include "PropertyPairsMap.h"
 #include "UObject/ObjectSaveContext.h"
 #include "UObject/Package.h"
 #include "UObject/UObjectBaseUtility.h" // GetNameSafe
 #include "UObject/UObjectIterator.h"
 
-#include UE_INLINE_GENERATED_CPP_BY_NAME(LandscapePatchComponent)
-
-#define LOCTEXT_NAMESPACE "LandscapePatch"
-
-namespace LandscapePatchComponentLocals
-{
 #if WITH_EDITOR
+#include "ScopedTransaction.h"
+#endif
 
-	ALandscapePatchManager* FindExistingPatchManagerForLandscape(ALandscape* Landscape)
-	{
-		for (int LayerIndex = 0; LayerIndex < Landscape->GetLayers().Num(); ++LayerIndex)
-		{
-			TArray<ALandscapeBlueprintBrushBase*> LayerBrushes = Landscape->GetBrushesForLayer(LayerIndex);
-			ALandscapeBlueprintBrushBase** Found = LayerBrushes.FindByPredicate(
-				[](ALandscapeBlueprintBrushBase* Candidate) { return Cast<ALandscapePatchManager>(Candidate) != nullptr; });
-			if (Found)
-			{
-				return Cast<ALandscapePatchManager>(*Found);
-			}
-		}
-
-		return nullptr;
-	}
-
-	ALandscapePatchManager* CreateNewPatchManagerForLandscape(ALandscape* Landscape)
-	{
-		if (!ensure(Landscape->CanHaveLayersContent()))
-		{
-			return nullptr;
-		}
+#include UE_INLINE_GENERATED_CPP_BY_NAME(LandscapePatchComponent)
 
-		FString BrushActorString = FString::Format(TEXT("{0}_{1}"), { Landscape->GetActorLabel(), ALandscapePatchManager::StaticClass()->GetName() });
-		FName BrushActorName = MakeUniqueObjectName(Landscape->GetOuter(), ALandscapePatchManager::StaticClass(), FName(BrushActorString));
-		FActorSpawnParameters SpawnParams;
-		SpawnParams.Name = BrushActorName;
-		SpawnParams.bAllowDuringConstructionScript = true; // This can be called by construction script if the actor being added to the world is part of a blueprint
+#define LOCTEXT_NAMESPACE "LandscapePatch"
 
-		ALandscapePatchManager* PatchManager = Landscape->GetWorld()->SpawnActor<ALandscapePatchManager>(ALandscapePatchManager::StaticClass(), SpawnParams);
-		if (PatchManager)
-		{
-			PatchManager->SetActorLabel(BrushActorString);
-			PatchManager->SetTargetLandscape(Landscape);
-		}
+/*
+ * About binding to landscape/edit layers:
+ * 
+ * EditLayerGuid: the guid of the edit layer that a patch wants to affect (the edit layer should be of type
+ *  ULandscapePatchEditLayer). The edit layer will hold a transient list of patches that have its guid, which
+ *  it should be notified to update whenever a patch is loaded/created/etc. The edit layer theoretically does
+ *  not need to be notified of *not* being pointed to, as it is able to filter its transient list whenever it
+ *  processes it, but we probably still want to issue a notification so that any UI displaying the patches can
+ *  update.
+ * DetailPanelLayerName: the name of the layer that the user sees. Should be kept in sync with EditLayerGuid.
+ * DetailPanelLayerGuid: the guid representation that the user sees. Should be kept in sync with EditLayerGuid.
+ * Priority: double that determines the sorting of the patches.
+ * 
+ * PatchManager: legacy object that used to hold a list of patches and apply them, where the patch order
+ *  was determined by the index in its legacy patch list. If PatchManager is not null, it takes precedence,
+ *  so the patch shouldn't be registered to an edit layer or use Priority for sorting. PatchManager should
+ *  be mutually exclusive with EditLayerGuid. So if PatchManager is set, it will clear the guid, and vice versa.
+ * 
+ * Landscape: pointer to the landscape in which the guid is found, or in which the PatchManager resides.
+ * 
+ * Resolving erroneous states:
+ * If landscape disagrees with guid or patch manager, guid/patch manager are authoritative, and landscape
+ *  should be fixed.
+ * If both PatchManager and guid are present, PatchManager wins.
+ * If PatchManager is not null but patch is not in the patch manager's list, this is an error, and patch
+ *  should be added to manager. However if patch is in a patch manager's list but does not have its PatchManager
+ *  pointer set, then the patch should be removed from the PatchManager. Basically, the patch, not the manager,
+ *  is authoritative.
+ * If the guid does not point to a valid edit layer, it does not get cleared (to be cleaner in cases where an
+ *  edit layer might be deleted and then the deletion undone), but the patch will not work until the guid is fixed.
+ * 
+ * Patch application:
+ * In non-batched-merge path, patches receive the legacy Initialize_Native and RenderLayer_Native calls. Legacy
+ *  list patches get this from the patch manager, and edit-layer-bound patches get it trhough the edit layer's
+ *  legacy adapter interface.
+ * In batched-merge path, patches (both in legacy list and in edit layer) provide their desired area of affect
+ *  and then get RenderLayer calls
+ */
 
-		return PatchManager;
-	}
 
-	bool bAllowAutoFindOfManagerInConstructionScript = true;
-	static FAutoConsoleVariableRef CVarReAddToPatchManagerOnConstructionIfLost(
-		TEXT("LandscapePatch.AllowAutoFindOfManagerInConstructionScript"),
-		bAllowAutoFindOfManagerInConstructionScript,
-		TEXT("If a blueprint-contained patch does not have a patch manager when its construction script is rerun, "
-			"allows for a patch manager to be found automatically. This is more safe to have turned off because "
-			"it can hide potential issues when patches or manager were not properly saved, and as a result, patch "
-			"order could change depending on load order of actors. It is on by default for now for legacy support."));
+namespace LandscapePatchComponentLocals
+{
+	FString NullDetailPanelLayerName = TEXT("-Null-");
 
+	const FText FixPatchBindingsTransactionName = LOCTEXT("FixPatchBindingsTransaction", "Fix Patch Bindings");
+	const FText RebindPatchTransactionName = LOCTEXT("RebindPatchTransaction", "Rebind Patch");
 
-	// A way to fix patch manager pointers. More aggressive than the map check, because it allows auto-finding
-	// of a new manager.
-	FAutoConsoleCommand CCmdFixPatchManagerPointers(
-		TEXT("LandscapePatch.FixPatchManagerPointers"),
-		TEXT("For all patches, make sure that patch manager pointers exist, contain this patch, and patches "
-			"and managers are dirtied if they changed. Patches without a manager get one auto-assigned."),
+#if WITH_EDITOR
+	FAutoConsoleCommand CCmdFixPatchBindings(
+		TEXT("LandscapePatch.FixPatchBindings"),
+		TEXT("For all patches, make sure that patch is either properly bound to a landscape layer or uses a legacy "
+			"patch manager."),
 		FConsoleCommandDelegate::CreateLambda([]() 
 	{
+		const FScopedTransaction Transaction(FixPatchBindingsTransactionName);
+
 		// Iterate through all patches
 		for (TObjectIterator<ULandscapePatchComponent> It(
 			/*AdditionalExclusionFlags = */RF_ClassDefaultObject,
@@ -93,37 +95,7 @@ namespace LandscapePatchComponentLocals
 				continue;
 			}
 
-			UWorld* World = Patch->GetWorld();
-			if (Patch->IsTemplate() || !IsValid(World) || World->WorldType != EWorldType::Editor)
-			{
-				continue;
-			}
-
-			Patch->SetUpPatchManagerData();
-			Patch->MarkDirtyIfModifiedInConstructionScript();
-		}
-
-		for (TObjectIterator<ALandscapePatchManager> It(
-			/*AdditionalExclusionFlags = */RF_ClassDefaultObject,
-			/*bIncludeDerivedClasses = */true,
-			/*InInternalExclusionFlags = */EInternalObjectFlags::Garbage); It; ++It)
-		{
-			ALandscapePatchManager* Manager = *It;
-			if (!IsValid(Manager))
-			{
-				continue;
-			}
-
-			UWorld* World = Manager->GetWorld();
-			if (Manager->IsTemplate() || !IsValid(World) || World->WorldType != EWorldType::Editor)
-			{
-				continue;
-			}
-
-			if (Manager)
-			{
-				Manager->MarkDirtyIfModifiedInConstructionScript();
-			}
+			Patch->FixBindings();
 		}
 	}));
 
@@ -159,8 +131,10 @@ void ULandscapePatchComponent::CheckForErrors()
 	Super::CheckForErrors();
 
 	using namespace LandscapePatchComponentLocals;
+
+	const FText FixPatchBindingsText = LOCTEXT("FixLandscapePointerButton", "Fix patch bindings");
 	
-	if (!IsRealPatch())
+	if (!IsPatchInWorld())
 	{
 		return;
 	}
@@ -173,60 +147,144 @@ void ULandscapePatchComponent::CheckForErrors()
 		return Arguments;
 	};
 
-	if (PatchManager.IsValid() != Landscape.IsValid())
+	if (PatchManager.IsValid())
 	{
-		FMessageLog("MapCheck").Warning()
+		if (EditLayerGuid.IsValid())
+		{
+			FMessageLog("MapCheck").Warning()
+			->AddToken(FUObjectToken::Create(this))
+			->AddToken(FTextToken::Create(FText::Format(LOCTEXT("HaveBothPatchManagerAndGuid", "Patch has both a patch "
+				"manager and a landscape edit layer guid. Patches should use one or the other. Guid is preferred because the "
+				"patch manager will be deprecated. (Package: {Package}, Actor: {Actor}). "
+				"Fix individually or run LandscapePatch.FixPatchBindings."), GetPackageAndActorArgs())))
+			->AddToken(FActionToken::Create(LOCTEXT("SwitchToUsingGuidButton", "Switch to using Guid"), FText(),
+				FOnActionTokenExecuted::CreateWeakLambda(this, [this]()
+			{
+				if (PatchManager.IsValid() && EditLayerGuid.IsValid())
+				{
+					const FScopedTransaction Transaction(RebindPatchTransactionName);
+					ResetPatchManager();
+				}
+			})));
+		}
+
+		if (PatchManager->GetOwningLandscape() != Landscape.Get())
+		{
+			FMessageLog("MapCheck").Warning()
 			->AddToken(FUObjectToken::Create(this))
 			->AddToken(FTextToken::Create(FText::Format(LOCTEXT("PatchManagerAndLandscapeDisagree", "Patch has inconsistent "
 				"manager and landscape pointers. (Package: {Package}, Actor: {Actor}). "
-				"Fix individually or run LandscapePatch.FixPatchManagerPointers."), GetPackageAndActorArgs())))
-			->AddToken(FActionToken::Create(LOCTEXT("FixInconsistentPointersButton", "Fix inconsistent pointers"), FText(),
-				FOnActionTokenExecuted::CreateWeakLambda(this, [this]() 
-		{
-			if (PatchManager.IsValid())
+				"Fix individually or run LandscapePatch.FixPatchBindings."), GetPackageAndActorArgs())))
+			->AddToken(FActionToken::Create(LOCTEXT("FixManagerLandscapeDisagreementButton", "Fix inconsistent pointers"), FText(),
+				FOnActionTokenExecuted::CreateWeakLambda(this, [this]()
 			{
-				SetPatchManager(PatchManager.Get());
-			}
-			else if (Landscape.IsValid())
+				if (PatchManager.IsValid() && PatchManager->GetOwningLandscape() != Landscape.Get())
+				{
+					const FScopedTransaction Transaction(RebindPatchTransactionName);
+
+					// It is tempting to just do SetPatchManager(PatchManager.Get()) here, but that will clear
+					// EditLayerGuid as well, which seems like it goes beyond what we promised to fix.
+					Modify();
+					Landscape = PatchManager->GetOwningLandscape();
+				}
+			})));
+		}
+
+		if (!PatchManager->ContainsPatch(this))
+		{
+			FMessageLog("MapCheck").Warning()
+			->AddToken(FUObjectToken::Create(this))
+			->AddToken(FTextToken::Create(FText::Format(LOCTEXT("PatchNotInManager", "Patch has patch manager, but "
+				"was not found in its legacy patch list. (Package: {Package}, Actor: {Actor}). "), GetPackageAndActorArgs())))
+			->AddToken(FActionToken::Create(LOCTEXT("AddToManagerButton", "Add to manager"), FText(),
+				FOnActionTokenExecuted::CreateWeakLambda(this, [this]() 
 			{
-				SetLandscape(Landscape.Get());
-			}
-		})));
+				if (PatchManager.IsValid())
+				{
+					const FScopedTransaction Transaction(RebindPatchTransactionName);
+					PatchManager->AddPatch(this);
+				}
+			})));
+		}
 	}
 
-	if (PatchManager.IsValid() && !PatchManager->ContainsPatch(this))
+	if (EditLayerGuid.IsValid())
 	{
-		FMessageLog("MapCheck").Warning()
+		const FLandscapeLayer* Layer = nullptr;
+		if (!Landscape.IsValid() || !Landscape->CanHaveLayersContent() || (Layer = Landscape->GetLayerConst(EditLayerGuid)) == nullptr)
+		{
+			FMessageLog("MapCheck").Warning()
 			->AddToken(FUObjectToken::Create(this))
-			->AddToken(FTextToken::Create(FText::Format(LOCTEXT("PatchNotInManager", "Patch was not found in its patch manager. "
-				"(Package: {Package}, Actor: {Actor}). "
-				"Fix individually or run LandscapePatch.FixPatchManagerPointers."), GetPackageAndActorArgs())))
-			->AddToken(FActionToken::Create(LOCTEXT("AddToManagerButton", "Add to manager"), FText(),
+			->AddToken(FTextToken::Create(FText::Format(LOCTEXT("GuidAndLandscapeDisagree", "The patch edit layer guid "
+				"did not match the Landscape pointer. (Package: {Package}, Actor: {Actor}). "
+				"Fix individually or run LandscapePatch.FixPatchBindings."), GetPackageAndActorArgs())))
+			->AddToken(FActionToken::Create(FixPatchBindingsText, FText(),
 				FOnActionTokenExecuted::CreateWeakLambda(this, [this]() 
+			{
+				const FScopedTransaction Transaction(FixPatchBindingsTransactionName);
+				FixBindings();
+			})));
+		}
+		else if (Cast<ULandscapePatchEditLayer>(Layer->EditLayer) == nullptr)
 		{
-			if (PatchManager.IsValid())
+			FMessageLog("MapCheck").Warning()
+			->AddToken(FUObjectToken::Create(this))
+			->AddToken(FTextToken::Create(FText::Format(LOCTEXT("GuidIsNotPatchLayer", "The patch edit layer guid points to a "
+				"layer that is not a landscape patch layer. (Package: {Package}, Actor: {Actor}). "
+				"Fix individually or run LandscapePatch.FixPatchBindings."), GetPackageAndActorArgs())))
+			->AddToken(FActionToken::Create(FixPatchBindingsText, FText(),
+				FOnActionTokenExecuted::CreateWeakLambda(this, [this]()
 			{
-				SetPatchManager(PatchManager.Get());
-			}
-		})));
+				const FScopedTransaction Transaction(FixPatchBindingsTransactionName);
+				FixBindings();
+			})));
+		}
 	}
 
-	if (bDirtiedByConstructionScript && !(GetPackage() && GetPackage()->IsDirty()))
+	if (!PatchManager.IsValid() && !EditLayerGuid.IsValid())
 	{
-		FMessageLog("MapCheck").Warning()
+		if (Landscape.IsValid() && !Landscape->CanHaveLayersContent())
+		{
+			FMessageLog("MapCheck").Warning()
+			->AddToken(FUObjectToken::Create(this))
+			->AddToken(FTextToken::Create(FText::Format(LOCTEXT("LandscapeEditLayersNotEnabled", "The patch points "
+				"to a landscape that does not have edit layers enabled. (Package: {Package}, Actor: {Actor}). "
+				"Fix individually or run LandscapePatch.FixPatchBindings."), GetPackageAndActorArgs())))
+			->AddToken(FActionToken::Create(FixPatchBindingsText, FText(),
+				FOnActionTokenExecuted::CreateWeakLambda(this, [this]() 
+			{
+				const FScopedTransaction Transaction(FixPatchBindingsTransactionName);
+				FixBindings();
+			})));
+		}
+		else if (Landscape.IsValid())
+		{
+			FMessageLog("MapCheck").Warning()
 			->AddToken(FUObjectToken::Create(this))
-			->AddToken(FTextToken::Create(FText::Format(LOCTEXT("PatchNeedsSaving", "Patch got modified in a construction "
-				"script rerun but has not been marked dirty. (Package: {Package}, Actor: {Actor}). "
-				"Fix individually or run LandscapePatch.FixPatchManagerPointers."), GetPackageAndActorArgs())))
-			->AddToken(FActionToken::Create(LOCTEXT("MarkDirtyButton", "Mark dirty"), FText(),
+			->AddToken(FTextToken::Create(FText::Format(LOCTEXT("PatchOnlyHasLandscape", "The patch had a landscape "
+				"but did not have an edit layer guid. (Package: {Package}, Actor: {Actor}). "
+				"Fix individually or run LandscapePatch.FixPatchBindings."), GetPackageAndActorArgs())))
+			->AddToken(FActionToken::Create(FixPatchBindingsText, FText(),
 				FOnActionTokenExecuted::CreateWeakLambda(this, [this]() 
+			{
+				const FScopedTransaction Transaction(FixPatchBindingsTransactionName);
+				FixBindings();
+			})));
+		}
+		else if (IsEnabled())
 		{
-			MarkPackageDirty();
-			if (PatchManager.IsValid())
+			FMessageLog("MapCheck").Warning()
+			->AddToken(FUObjectToken::Create(this))
+			->AddToken(FTextToken::Create(FText::Format(LOCTEXT("EnabledPatchNotBoundToLandscape", "Patch is enabled "
+				"but is not bound to a landscape edit layer. (Package: {Package}, Actor: {Actor}). "
+				"Fix individually or run LandscapePatch.FixPatchBindings."), GetPackageAndActorArgs())))
+			->AddToken(FActionToken::Create(FixPatchBindingsText, FText(),
+				FOnActionTokenExecuted::CreateWeakLambda(this, [this]() 
 			{
-				PatchManager->MarkDirtyIfModifiedInConstructionScript();
-			}
-		})));
+				const FScopedTransaction Transaction(FixPatchBindingsTransactionName);
+				FixBindings();
+			})));
+		}
 	}
 }
 
@@ -236,37 +294,17 @@ void ULandscapePatchComponent::OnComponentCreated()
 
 	Super::OnComponentCreated();
 
-	// Mark whether we're creating from scratch or from a copy
 	bWasCopy = bPropertiesCopiedIndicator;
 	bPropertiesCopiedIndicator = true;
 
-	if (!IsRealPatch())
-	{
-		return;
-	}
-
-	// Doing stuff during construction script reruns is a huge pain, and ideally we would always exit right here.
-	// However we preserve legacy behavior where patches added themselves to the manager from OnComponentCreated
-	// even in the construction script, if bAllowAutoFindOfManagerInConstructionScript is true.
-	// See ApplyComponentInstanceData to see how we try to detect that patch got modified in the construction
-	// script and therefore needs dirtying.
-	if (GIsReconstructingBlueprintInstances && !bAllowAutoFindOfManagerInConstructionScript)
-	{
-		return;
-	}
-
-	SetUpPatchManagerData();
-}
-
-void ULandscapePatchComponent::SetUpPatchManagerData()
-{
-	using namespace LandscapePatchComponentLocals;
-
-	if (!IsRealPatch())
+	// Doing stuff during construction script reruns is a huge pain. Avoid it.
+	if (GIsReconstructingBlueprintInstances)
 	{
 		return;
 	}
 
+	// Legacy: if we're registered through a patch manager (i.e. we're copying a component that is),
+	// make sure our other pointers are consistent.
 	if (PatchManager.IsValid())
 	{
 		if (Landscape != PatchManager->GetOwningLandscape())
@@ -276,94 +314,95 @@ void ULandscapePatchComponent::SetUpPatchManagerData()
 		}
 
 		// The patch manager might legitimately not contain the set patch manager if we're copying a patch. 
-		if (!PatchManager->ContainsPatch(this))
+		if (!PatchManager->ContainsPatch(this) && IsPatchInWorld())
 		{
-			// Note: in blueprint-contained patches, this can actually fix a case where our patch manager was
-			// saved incorrectly, but we need the patch manager to be dirtied if we're failing to dirty due to
-			// a load.
-			if (GIsReconstructingBlueprintInstances)
-			{
-				UE_LOG(LogLandscapePatch, Warning, TEXT("Patch was not in the set Patch Manager and was added in a construction script rerun. "
-					"The patch manager should be saved to avoid potential reorderings depending on loading order. (Package: %s, Manager: %s"),
-					*GetNameSafe(PatchManager->GetPackage()), *GetNameSafe(PatchManager.Get()));
-			}
-
 			PatchManager->AddPatch(this);
 		}
 
-		// No more adjustments needed
+		ResetEditLayer();
 		return;
 	}
 
-	// If we got here, we don't have a patch manager. See if we have a landscape to get it from.
-	if (Landscape.IsValid())
+	// Otherwise, bind to some edit layer
+	bool bConnectionToLandscapeEstablished = false;
+	if (EditLayerGuid.IsValid())
 	{
-		SetLandscape(Landscape.Get());
-		return;
+		bConnectionToLandscapeEstablished = BindToEditLayer(EditLayerGuid);
 	}
-
-	// See if the level has a height patch manager to which we can add ourselves
-	UWorld* World = GetWorld();
-	if (!ensure(World))
+	if (!bConnectionToLandscapeEstablished && Landscape.IsValid())
 	{
-		return;
+		// Minor note: the above BindToEditLayer can fail and yet still change the landscape pointer if the
+		//  guid pointed to a real layer in a different landscape but of the wrong type. This behavior is
+		//  probably desirable (presumably the user wanted to bind to that landscape?) but hard to say, and
+		//  unlikely to come up in the first place.
+		bConnectionToLandscapeEstablished = BindToLandscape(Landscape.Get());
 	}
-	TActorIterator<ALandscapePatchManager> ManagerIterator(World);
-	if (!!ManagerIterator)
+	if (!bConnectionToLandscapeEstablished && IsPatchInWorld())
 	{
-		SetPatchManager(*ManagerIterator);
+		BindToAnyLandscape();
 	}
-	else
+
+	// Update priority
+	if (IsPatchInWorld())
 	{
-		// If no existing manager, find some landscape and add a new one.
-		for (TActorIterator<ALandscape> LandscapeIterator(World); LandscapeIterator; ++LandscapeIterator)
+		switch (PriorityInitialization)
 		{
-			if (LandscapeIterator->CanHaveLayersContent())
+		case ELandscapePatchPriorityInitialization::KeepOriginal:
+			// Do nothing
+			break;
+		case ELandscapePatchPriorityInitialization::SmallIncrement:
+			Priority += 0.01;
+			break;
+		case ELandscapePatchPriorityInitialization::AcquireHighest:
+			if (EditLayer.IsValid())
 			{
-				Modify();
-				Landscape = *LandscapeIterator;
-				SetPatchManager(CreateNewPatchManagerForLandscape(Landscape.Get()));
-				break;
+				if (IsPatchPreview())
+				{
+					Priority = EditLayer->GetHighestPatchPriority();
+				}
+				else
+				{
+					Priority = EditLayer->ClaimHighestPatchPriority();
+				}
 			}
+			break;
 		}
 	}
-	if (!PatchManager.IsValid())
-	{
-		UE_LOG(LogLandscapePatch, Warning, TEXT("Unable to find a landscape with edit layers enabled. Unable to create patch manager."));
-	}
-}
-
-void ULandscapePatchComponent::MarkDirtyIfModifiedInConstructionScript()
-{
-	if (IsRealPatch() && bDirtiedByConstructionScript)
-	{
-		MarkPackageDirty();
-	}
 }
 
 void ULandscapePatchComponent::PostLoad()
 {
 	Super::PostLoad();
-
-	bPropertiesCopiedIndicator = true;
-	bLoadedButNotYetRegistered = true;
 }
 
-void ULandscapePatchComponent::PreSave(FObjectPreSaveContext SaveContext)
+void ULandscapePatchComponent::PostEditUndo()
 {
-	Super::PreSave(SaveContext);
+	Super::PostEditUndo();
+
+	// Keeps all of our guid-related variables consistent
+	BindToEditLayer(EditLayerGuid);
 
-	// If we're ssaving, then we no longer have to worry about our dirtiness.
-	bDirtiedByConstructionScript = false;
+	// Makes sure we update regardless of whether we're using patch manager or guid
+	RequestLandscapeUpdate();
 }
 
 void ULandscapePatchComponent::OnComponentDestroyed(bool bDestroyingHierarchy)
 {
-	if (PatchManager.IsValid() && !GIsReconstructingBlueprintInstances)
+	if (!GIsReconstructingBlueprintInstances)
 	{
-		PatchManager->RemovePatch(this);
+		if (PatchManager.IsValid())
+		{
+			PatchManager->RemovePatch(this);
+		}
+		else if (EditLayer.IsValid())
+		{
+			// Notify the layer that the patch is being destroyed. Note that we are not yet marked
+			//  garbage, so clear our edit layer guid so that the edit layer correctly sees us as
+			//  disconnected.
+			ULandscapePatchEditLayer* EditLayerToNotify = EditLayer.Get();
+			ResetEditLayer();
+		}
 	}
-
 	Super::OnComponentDestroyed(bDestroyingHierarchy);
 }
 
@@ -371,17 +410,17 @@ void ULandscapePatchComponent::OnRegister()
 {
 	Super::OnRegister();
 
-
-	if (bLoadedButNotYetRegistered)
+	if (!IsPatchInWorld())
 	{
-		bLoadedButNotYetRegistered = false;
 		return;
 	}
 
-	// TODO: Currently the main reason to invalidate the landscape on registration is to respond
-	// to detail panel changes of the parent's transform. However we may be able to catch a wide
-	// variety of changes this way, so we'll need to see if we can get rid of other invalidations.
-	// Also, we should make the invalidation conditional on whether we actually modify any relevant
+	if (!GIsReconstructingBlueprintInstances && EditLayerGuid.IsValid())
+	{
+		BindToEditLayer(EditLayerGuid);
+	}
+
+	// TODO: We should make the invalidation conditional on whether we actually modify any relevant
 	// properties by having a virtual method that compares and updates a stored hash of them.
 	if (IsEnabled())
 	{
@@ -408,9 +447,9 @@ void ULandscapePatchComponent::OnUpdateTransform(EUpdateTransformFlags UpdateTra
 {
 	Super::OnUpdateTransform(UpdateTransformFlags, Teleport);
 
-	if (IsEnabled())
+	if (IsEnabled() && CanAffectLandscape())
 	{
-		RequestLandscapeUpdate(/* bInUserTriggeredUpdate = */ true);
+		RequestLandscapeUpdate();
 	}
 }
 
@@ -418,12 +457,6 @@ void ULandscapePatchComponent::PostEditChangeProperty(FPropertyChangedEvent& Pro
 {
 	using namespace LandscapePatchComponentLocals;
 
-	// Do a bunch of checks to make sure that we don't try to do anything when the editing is happening inside the blueprint editor.
-	if (!IsRealPatch())
-	{
-		return;
-	}
-
 	// If we're changing the owning landscape or patch manaager, there's some work we need to do to remove/add 
 	// ourselves from/to the proper brush managers.
 	if (PropertyChangedEvent.Property 
@@ -434,13 +467,31 @@ void ULandscapePatchComponent::PostEditChangeProperty(FPropertyChangedEvent& Pro
 	else if (PropertyChangedEvent.Property
 		&& (PropertyChangedEvent.Property->GetFName() == GET_MEMBER_NAME_CHECKED(ULandscapePatchComponent, PatchManager)))
 	{
+		PRAGMA_DISABLE_DEPRECATION_WARNINGS
 		SetPatchManager(PatchManager.Get());
+		PRAGMA_ENABLE_DEPRECATION_WARNINGS
+	}
+	else if (PropertyChangedEvent.Property
+		&& (PropertyChangedEvent.Property->GetFName() == GET_MEMBER_NAME_CHECKED(ULandscapePatchComponent, DetailPanelLayerName)))
+	{
+		UpdateEditLayerFromDetailPanelLayerName();
+	}
+	else if (PropertyChangedEvent.Property
+		&& (PropertyChangedEvent.Property->GetFName() == GET_MEMBER_NAME_CHECKED(ULandscapePatchComponent, Priority)))
+	{
+		// We don't use SetPriority because that does nothing if the priority does not change, and by this point
+		// the value has been updated. All we're missing is the extra notification.
+		if (EditLayer.IsValid())
+		{
+			EditLayer->NotifyOfPriorityChange(this);
+		}
 	}
 
-	if (IsEnabled() || (PropertyChangedEvent.Property
-		&& (PropertyChangedEvent.Property->GetFName() == GET_MEMBER_NAME_CHECKED(ULandscapePatchComponent, bIsEnabled))))
+	// Request a landscape update as long as we're enabled, or if we just disabled ourselves.
+	if (IsPatchInWorld() && (IsEnabled() || (PropertyChangedEvent.Property
+		&& (PropertyChangedEvent.Property->GetFName() == GET_MEMBER_NAME_CHECKED(ULandscapePatchComponent, bIsEnabled)))))
 	{
-		RequestLandscapeUpdate(/* bInUserTriggeredUpdate = */ true);
+		RequestLandscapeUpdate();
 	}
 
 	// It is important that this super call happen after the above, because inside a blueprint actor, the call triggers a
@@ -462,21 +513,26 @@ FLandscapePatchComponentInstanceData::FLandscapePatchComponentInstanceData(const
 		return;
 	}
 
-	bGaveMissingPatchManagerWarning = Patch->bGaveMissingPatchManagerWarning;
+	bGaveCouldNotBindToEditLayerWarning = Patch->bGaveCouldNotBindToEditLayerWarning;
+	bGaveMismatchedLandscapeWarning = Patch->bGaveMismatchedLandscapeWarning;
+	bGaveMissingEditLayerGuidWarning = Patch->bGaveMissingEditLayerGuidWarning;
 	bGaveNotInPatchManagerWarning = Patch->bGaveNotInPatchManagerWarning;
 	bGaveMissingLandscapeWarning = Patch->bGaveMissingLandscapeWarning;
 
-	bDirtiedByConstructionScript = Patch->bDirtiedByConstructionScript;
-
+	PRAGMA_DISABLE_DEPRECATION_WARNINGS
 	PatchManager = Patch->GetPatchManager();
+	PRAGMA_ENABLE_DEPRECATION_WARNINGS
+
+	EditLayerGuid = Patch->EditLayerGuid;
+	Priority = Patch->Priority;
 #endif
 }
 
+#if WITH_EDITOR
 // Called before/after rerunning construction scripts (when patch is part of a blueprint) to carry over extra data.
 void ULandscapePatchComponent::ApplyComponentInstanceData(FLandscapePatchComponentInstanceData* ComponentInstanceData, 
 	ECacheApplyPhase CacheApplyPhase)
 {
-#if WITH_EDITOR
 	using namespace LandscapePatchComponentLocals;
 
 	if (CacheApplyPhase == ECacheApplyPhase::PostUserConstructionScript)
@@ -485,260 +541,610 @@ void ULandscapePatchComponent::ApplyComponentInstanceData(FLandscapePatchCompone
 		return;
 	}
 
-	if (!ComponentInstanceData || !IsRealPatch())
+	if (!ComponentInstanceData)
 	{
 		return;
 	}
 
-	bGaveMissingPatchManagerWarning = ComponentInstanceData->bGaveMissingPatchManagerWarning;
+	bGaveCouldNotBindToEditLayerWarning = ComponentInstanceData->bGaveCouldNotBindToEditLayerWarning;
+	bGaveMismatchedLandscapeWarning = ComponentInstanceData->bGaveMismatchedLandscapeWarning;
+	bGaveMissingEditLayerGuidWarning = ComponentInstanceData->bGaveMissingEditLayerGuidWarning;
 	bGaveNotInPatchManagerWarning = ComponentInstanceData->bGaveNotInPatchManagerWarning;
 	bGaveMissingLandscapeWarning = ComponentInstanceData->bGaveMissingLandscapeWarning;
 
-	bDirtiedByConstructionScript = ComponentInstanceData->bDirtiedByConstructionScript;
+	EditLayerGuid = ComponentInstanceData->EditLayerGuid;
+	Priority = ComponentInstanceData->Priority;
 
-	if (!bAllowAutoFindOfManagerInConstructionScript)
+	PatchManager = ComponentInstanceData->PatchManager;
+	PreviousPatchManager = PatchManager;
+	if (PatchManager.IsValid())
 	{
-		// This is the ideal case where we don't have to worry about auto acquired manager and we
-		// don't expect to have to modify the previous one in any way.
-		PatchManager = ComponentInstanceData->PatchManager;
-		PreviousPatchManager = PatchManager;
-		Landscape = PatchManager.IsValid() ? PatchManager->GetOwningLandscape() : nullptr;
+		Landscape = PatchManager->GetOwningLandscape();
 
-		if (PatchManager.IsValid())
+		if (IsPatchInWorld() && !PatchManager->ContainsPatch(this))
 		{
-			if (!ensure(PatchManager->ContainsPatch(this)))
-			{
-				PatchManager->AddPatch(this);
-				PatchManager->MarkModifiedInConstructionScript();
-			}
+			PatchManager->AddPatch(this);
 		}
-		return;
+	}
+	else
+	{
+		BindToEditLayer(EditLayerGuid);
 	}
 
-	// If we got here, then we're going down a legacy path where we auto acquire patch manager
-	// in construction script.
-
-	if (ComponentInstanceData->PatchManager.IsValid())
+	bInstanceDataApplied = true;
+	if (bDeferUpdateRequestUntilInstanceData)
 	{
-		// We had an existing patch manager so use that. This is mostly like the above except that we
-		// have to remove ourselves from an acquired manager if there is one, and we automatically
-		// add ourselves to our patch manager if it was saved without this patch. Calling SetPatchManager
-		// does all of that for us, except we need to track whether we modified the patch manager we keep.
+		RequestLandscapeUpdate();
+		bDeferUpdateRequestUntilInstanceData = false;
+	}
+}
+#endif // WITH_EDITOR
+
+void ULandscapePatchComponent::SetLandscape(ALandscape* NewLandscape)
+{
+#if WITH_EDITOR
 
-		bool bModifiedPatchManager = !ComponentInstanceData->PatchManager->ContainsPatch(this);
+	using namespace LandscapePatchComponentLocals;
 
-		SetPatchManager(ComponentInstanceData->PatchManager.Get());
+	Modify(false);
+	Landscape = NewLandscape;
 
-		if (bModifiedPatchManager)
-		{
-			ComponentInstanceData->PatchManager->MarkModifiedInConstructionScript();
-
-			// This might get double logged because InstanceDataCache->ApplyToActor gets run twice (see AActor::ExecuteConstruction),
-			// but that's ok. Also we could check for current dirtyness to avoid the "might not have" wording,
-			// but we want to log a warning either way, so probably not worth it.
-			UE_LOG(LogLandscapePatch, Warning, TEXT("Patch was missing from expected patch manager and was added in a "
-				"construction script rerun but manager might not have been marked dirty. Save the manager to avoid "
-				"potential reorderings depending on loading order. (Package: %s, Actor: %s"),
-				*GetNameSafe(GetPackage()), *GetNameSafe(GetAttachmentRootActor()));
-		}
+	if (!NewLandscape)
+	{
+		ResetEditLayer();
+		ResetPatchManager();
 		return;
 	}
 
-	// If we got to here, then we're keeping the patch manager we auto acquired.
 	if (PatchManager.IsValid())
 	{
-		bDirtiedByConstructionScript = true;
-		PatchManager->MarkModifiedInConstructionScript();
+		if (PatchManager->GetOwningLandscape() == Landscape)
+		{
+			// We're already bound to this landscape through the patch manager
+			return;
+		}
+		// Otherwise, disconnect from the patch manager
+		ResetPatchManager();
+	}
 
-		// See comment on the other log statement above
-		UE_LOG(LogLandscapePatch, Warning, TEXT("Patch acquired a patch manager in a construction script rerun but might not have been "
-			"marked dirty. Save the patch and its manager to avoid potential reorderings depending on loading order. (Package: %s, Actor: %s"),
-			*GetNameSafe(GetPackage()), *GetNameSafe(GetAttachmentRootActor()));
+	if (!BindToLandscape(NewLandscape))
+	{
+		UE_LOG(LogLandscapePatch, Warning, TEXT("Unable to bind to given landscape (does it have edit layers enabled?)."));
 	}
-	
 #endif // WITH_EDITOR
 }
 
-void ULandscapePatchComponent::SetLandscape(ALandscape* NewLandscape)
+// Changes landscape pointer and layer guid only if given manager is not null
+void ULandscapePatchComponent::SetPatchManager(ALandscapePatchManager* NewPatchManager)
 {
 #if WITH_EDITOR
+	if (!NewPatchManager)
+	{
+		ResetPatchManager();
+		return;
+	}
 
-	using namespace LandscapePatchComponentLocals;
+	Modify(false);
+	
+	// If we're setting a patch manager, we shouldn't be using the normal guid/priority approach.
+	ResetEditLayer();
 
-	// We early out here if the landscape stays the same and we have a patch manager because it is inconvenient to accidentally swap patch
-	// managers if there are multiple in the same landscape.
-	if (Landscape.Get() == NewLandscape && (
-		(Landscape.IsNull() && PatchManager.IsNull()) || (Landscape.IsValid() && PatchManager.IsValid() && PatchManager->GetOwningLandscape() == Landscape.Get())))
+	if (PreviousPatchManager.IsValid() && PreviousPatchManager != NewPatchManager)
 	{
-		return;
+		PreviousPatchManager->RemovePatch(this);
 	}
 
-	Modify();
+	PatchManager = NewPatchManager;
+	if (IsPatchInWorld())
+	{
+		PatchManager->AddPatch(this);
+	}
+	Landscape = PatchManager->GetOwningLandscape();
+	PreviousPatchManager = NewPatchManager;
+	ResetWarnings();
+#endif // WITH_EDITOR
+}
 
-	Landscape = NewLandscape;
+#if WITH_EDITOR
+// Does not touch landscape pointer or layer guid
+void ULandscapePatchComponent::ResetPatchManager()
+{
+	Modify(false);
 
-	if (!NewLandscape)
+	if (PreviousPatchManager.IsValid())
+	{
+		PreviousPatchManager->RemovePatch(this);
+	}
+	PatchManager = nullptr;
+	PreviousPatchManager = nullptr;
+}
+
+bool ULandscapePatchComponent::IsPatchPreview()
+{
+	return GetOwner() && GetOwner()->bIsEditorPreviewActor;
+}
+#endif
+
+ALandscapePatchManager* ULandscapePatchComponent::GetPatchManager() const
+{
+	return PatchManager.Get();
+}
+
+void ULandscapePatchComponent::RequestLandscapeUpdate(bool bInUserTriggeredUpdate /* = false */)
+{
+	// TODO: Once we're sure that the bool parameter is not necessary, we should say so in the function
+	// header. Might not be able to remove safely since it's blueprint callable...
+
+	using namespace LandscapePatchComponentLocals;
+
+#if WITH_EDITOR
+	// Note that aside from the usual guard against doing things in the blueprint editor, the check of WorldType
+	// inside this call also prevents us from doing the request while cooking, where WorldType is set to Inactive. Otherwise
+	// we might issue warnings below.
+	if (!IsPatchInWorld())
 	{
-		SetPatchManager(nullptr);
 		return;
 	}
 
-	if (!IsRealPatch())
+	// If we get a request for a landscape during rerunning construction scripts before applying instance data,
+	// defer that request until we've applied the instance data. If we don't, then the below booleans that are
+	// meant to stop spamming the log will not work, since they get carried over via instance data.
+	if (GIsReconstructingBlueprintInstances && !bInstanceDataApplied)
 	{
+		bDeferUpdateRequestUntilInstanceData = true;
 		return;
 	}
 
-	// If landscape was valid, try to find a patch manager inside that landscape.
-	if (Landscape->CanHaveLayersContent())
+	if (PatchManager.IsValid())
 	{
-		ALandscapePatchManager* ManagerToUse = FindExistingPatchManagerForLandscape(Landscape.Get());
-		if (!ManagerToUse)
+		// Legacy path of requesting an update through a patch manager
+
+		bool bRequestUpdate = true;
+		if (!PatchManager->ContainsPatch(this))
 		{
-			ManagerToUse = CreateNewPatchManagerForLandscape(Landscape.Get());
+			if (!bGaveNotInPatchManagerWarning)
+			{
+				UE_LOG(LogLandscapePatch, Warning, TEXT("Patch's patch manager does not contain this patch. "
+					"Perhaps the manager was not saved? Reset the patch manager on the patch. (Package: %s, Actor: %s)"),
+					*GetNameSafe(GetPackage()), *GetNameSafe(GetAttachmentRootActor()));
+				bGaveNotInPatchManagerWarning = true;
+			}
+			bRequestUpdate = false;
+		}
+		if (!IsValid(PatchManager->GetOwningLandscape()))
+		{
+			if (!bGaveMissingLandscapeWarning)
+			{
+				UE_LOG(LogLandscapePatch, Warning, TEXT("Patch's patch manager does not have a valid owning "
+					"landscape. Perhaps the landscape was not saved? Reset the landscape on the manager. (Package: %s, Manager: %s)"),
+					*GetNameSafe(GetPackage()), *GetNameSafe(PatchManager.Get()));
+				bGaveMissingLandscapeWarning = true;
+			}
+			bRequestUpdate = false;
 		}
 
-		SetPatchManager(ManagerToUse);
+		if (bRequestUpdate)
+		{
+			ResetWarnings();
+			PatchManager->RequestLandscapeUpdate(bInUserTriggeredUpdate || !UE::GetIsEditorLoadingPackage());
+		}
+
+		return;
 	}
-	else
+
+	// Otherwise, we work via Guid registration.
+	if (!EditLayer.IsValid() && EditLayerGuid.IsValid())
 	{
-		UE_LOG(LogLandscapePatch, Warning, TEXT("Landscape target for height patch did not have edit layers enabled. Unable to create patch manager."));
-		Landscape = nullptr;
-		SetPatchManager(nullptr);
+		BindToEditLayer(EditLayerGuid);
 	}
 
-#endif // WITH_EDITOR
+	if (EditLayer.IsValid())
+	{
+		ResetWarnings();
+		EditLayer->RequestLandscapeUpdate();
+	}
+	else if (EditLayerGuid.IsValid())
+	{
+		// We have a guid, but weren't able to bind to it
+		if (!bGaveCouldNotBindToEditLayerWarning)
+		{
+			UE_LOG(LogLandscapePatch, Warning, TEXT("Could not find an edit layer with the given guid for the patch. "
+				"(Package: %s, Manager: %s)"),
+				*GetNameSafe(GetPackage()), *GetNameSafe(PatchManager.Get()));
+			bGaveCouldNotBindToEditLayerWarning = true;
+		}
+	}
+	else
+	{
+		// We didn't even have a guid
+		if (!bGaveMissingEditLayerGuidWarning)
+		{
+			UE_LOG(LogLandscapePatch, Warning, TEXT("Patch is not bound to an edit layer. Set the landscape and edit layer guid "
+				"on the patch. (Package: %s)"),
+				*GetNameSafe(GetPackage()));
+			bGaveMissingEditLayerGuidWarning = true;
+		}
+	}
+#endif
 }
 
-void ULandscapePatchComponent::SetPatchManager(ALandscapePatchManager* NewPatchManager)
+bool ULandscapePatchComponent::IsPatchInWorld() const
 {
-	// Could have put the WITH_EDITOR block around just GetOwningLandscape, but might as well
-	// surround everything since nothing is currently expected to work for a patch manager at runtime.
+	UWorld* World = GetWorld();
+	return !IsTemplate() && IsValid(this) && IsValid(World) && World->WorldType == EWorldType::Editor;
+}
+
 #if WITH_EDITOR
+void ULandscapePatchComponent::ResetWarnings()
+{
+	bGaveCouldNotBindToEditLayerWarning = false;
+	bGaveMismatchedLandscapeWarning = false;
+	bGaveMissingEditLayerGuidWarning = false;
+	bGaveNotInPatchManagerWarning = false;
+	bGaveMissingLandscapeWarning = false;
+}
 
-	if (PreviousPatchManager == NewPatchManager 
-		&& PatchManager == NewPatchManager
-		&& !(NewPatchManager && !NewPatchManager->ContainsPatch(this)))
+// Safe to do if already bound to the given layer, and safe to do for templates. Does not
+//  touch patch manager. Changes the landscape pointer only if the guid points to a layer, even
+//  if the layer is of the wrong type. Changes all guid-related variables regardless of whether 
+//  the guid points to a valid layer or not. EditLayerGuid will be valid only if the guid points
+//  to a layer of an appropriate type (in which case binding is considered successful). Binding
+//  to null guid is same as a reset, and is considered successful.
+bool ULandscapePatchComponent::BindToEditLayer(const FGuid& Guid)
+{
+	using namespace LandscapePatchComponentLocals;
+
+	if (!Guid.IsValid())
 	{
-		return;
+		ResetEditLayer();
+		return true;
 	}
 
-	ResetWarnings();
+	Modify(false);
+	ULandscapePatchEditLayer* PreviousLayer = EditLayer.Get();
 
-	if (PreviousPatchManager.IsValid())
+	const FLandscapeLayer* LayerStruct = nullptr;
+
+	// See if the layer is in our current landscape
+	if (Landscape.Get())
 	{
-		PreviousPatchManager->RemovePatch(this);
+		LayerStruct = Landscape->GetLayerConst(Guid);
+		if (!LayerStruct && !bGaveMismatchedLandscapeWarning)
+		{
+			UE_LOG(LogLandscapePatch, Warning, TEXT("Mismatch between landscape and layer Guid in patch."));
+			bGaveMismatchedLandscapeWarning = true;
+		}
 	}
-
-	Modify();
-	PatchManager = NewPatchManager;
-	if (NewPatchManager)
+	
+	// If the layer wasn't in the current landscape, see if it's in some other landsape
+	if (!LayerStruct)
 	{
-		if (IsRealPatch())
+		UWorld* World = GetWorld();
+		if (ensure(World))
 		{
-			PatchManager->AddPatch(this);
+			for (TActorIterator<ALandscape> LandscapeIterator(World); LandscapeIterator; ++LandscapeIterator)
+			{
+				LayerStruct = LandscapeIterator->GetLayerConst(Guid);
+				if (LayerStruct)
+				{
+					// Found!
+					// Note that the layer we found might be of an incorrect type, so it's arguable whether landscape
+					//  pointer should change yet. But probably should.
+					Landscape = *LandscapeIterator;
+					break;
+				}
+			}
 		}
-		Landscape = PatchManager->GetOwningLandscape();
 	}
-	else
+
+	EditLayer = Cast<ULandscapePatchEditLayer>(LayerStruct ? LayerStruct->EditLayer : nullptr);
+	EditLayerGuid = Guid;
+	DetailPanelLayerGuid = Guid.ToString();
+	DetailPanelLayerName = EditLayer.IsValid() ? LayerStruct->Name.ToString()
+		: LayerStruct ? TEXT("-Incorrect Layer Type-")
+		: TEXT("-Layer Not Found-");
+
+	if (PreviousLayer && EditLayer != PreviousLayer)
 	{
-		Landscape = nullptr;
+		PreviousLayer->NotifyOfPatchRemoval(this);
+		PreviousLayer->RequestLandscapeUpdate();
 	}
 
-	PreviousPatchManager = NewPatchManager;
+	if (EditLayer.IsValid() && IsPatchInWorld())
+	{
+		// This is safe to do even if we are already registered
+		EditLayer->RegisterPatchForEditLayer(this);
+	}
 
-#endif // WITH_EDITOR
+	return EditLayer.IsValid();
 }
 
-ALandscapePatchManager* ULandscapePatchComponent::GetPatchManager() const
+// Does not affect landscape pointer
+void ULandscapePatchComponent::ResetEditLayer()
 {
-	return PatchManager.Get();
+	using namespace LandscapePatchComponentLocals;
+
+	ULandscapePatchEditLayer* PreviousLayer = EditLayer.Get();
+
+	Modify(false);
+
+	EditLayerGuid.Invalidate();
+	DetailPanelLayerGuid = EditLayerGuid.ToString();
+	DetailPanelLayerName = NullDetailPanelLayerName;
+	EditLayer = nullptr;
+
+	if (PreviousLayer)
+	{
+		PreviousLayer->NotifyOfPatchRemoval(this);
+		PreviousLayer->RequestLandscapeUpdate();
+	}
 }
 
-void ULandscapePatchComponent::MoveToTop()
+// If we have an edit layer guid in the given landscape, binds to that edit layer. Otherwise,
+// looks for an edit layer to bind to in that landscape.
+bool ULandscapePatchComponent::BindToLandscape(ALandscape* LandscapeIn)
 {
-	if (PatchManager.IsValid() && PatchManager->ContainsPatch(this))
+	using namespace LandscapePatchComponentLocals;
+	
+	Modify(false);
+	Landscape = LandscapeIn;
+
+	if (!LandscapeIn)
 	{
-		PatchManager->Modify();
-		PatchManager->RemovePatch(this);
-		PatchManager->AddPatch(this);
-		if (IsEnabled())
+		ResetPatchManager();
+		ResetEditLayer();
+		return true;
+	}
+
+	if (!LandscapeIn->CanHaveLayersContent())
+	{
+		ResetPatchManager();
+		ResetEditLayer();
+		return false;
+	}
+
+	// See if we're already bound via a patch manager. If not, make sure patch manager is reset.
+	if (PatchManager.IsValid() && PatchManager->GetOwningLandscape() == LandscapeIn)
+	{
+		ResetEditLayer();
+		return true;
+	}
+	ResetPatchManager();
+
+	const FLandscapeLayer* LayerStruct = nullptr;
+	ULandscapePatchEditLayer* Layer = nullptr;
+	if (EditLayerGuid.IsValid())
+	{
+		// See if this layer is here and of the appropriate type
+		LayerStruct = LandscapeIn->GetLayerConst(EditLayerGuid);
+		Layer = LayerStruct ? Cast<ULandscapePatchEditLayer>(LayerStruct->EditLayer) : nullptr;
+	}
+	if (!Layer)
+	{
+		// See if we have any layer of the appropriate type
+		LayerStruct = LandscapeIn->FindLayerOfType(ULandscapePatchEditLayer::StaticClass());
+		Layer = LayerStruct ? Cast<ULandscapePatchEditLayer>(LayerStruct->EditLayer) : nullptr;
+	}
+
+	// If we couldn't find an appropriate layer, make a new one
+	if (!Layer && IsPatchInWorld() 
+		// We don't want to create a layer if this is the preview, because that will make the layer insertion
+		// not transactable, and we're not going to delete the layer if we don't complete the drop.
+		&& !IsPatchPreview())
+	{
+		const FName PatchLayerName = LandscapeIn->GenerateUniqueLayerName(
+			*ULandscapePatchEditLayer::StaticClass()->GetDefaultObject<ULandscapePatchEditLayer>()->GetDefaultName());
+
+		ILandscapeModule& LandscapeModule = FModuleManager::GetModuleChecked<ILandscapeModule>("Landscape");
+		int32 LayerIndex = LandscapeModule.GetLandscapeEditorServices()->GetOrCreateEditLayer(PatchLayerName, LandscapeIn, ULandscapePatchEditLayer::StaticClass());
+		LayerStruct = LandscapeIn->GetLayerConst(LayerIndex);
+
+		Layer = LayerStruct ? Cast<ULandscapePatchEditLayer>(LayerStruct->EditLayer) : nullptr;
+	}
+
+	if (!Layer)
+	{
+		// This happens if we failed to find a layer and either the patch was a preview, or IsPatchInWorld was false
+		// (the latter has happened while doing a SaveAs on a level).
+		ResetEditLayer();
+		return false;
+	}
+
+	return BindToEditLayer(LayerStruct->Guid);
+}
+
+bool ULandscapePatchComponent::BindToAnyLandscape()
+{
+	// While we typically allow the bind functions to operate on templates so that all members
+	//  of the template can be initialized coherently, we don't want to accidentally bind
+	//  templates to a landscape automatically.
+	// So, early out in this particular function.
+	if (!IsPatchInWorld())
+	{
+		return false;
+	}
+
+	UWorld* World = GetWorld();
+	if (!ensure(World))
+	{
+		return false;
+	}
+
+	for (TActorIterator<ALandscape> LandscapeIterator(World); LandscapeIterator; ++LandscapeIterator)
+	{
+		if (LandscapeIterator->CanHaveLayersContent() && BindToLandscape(*LandscapeIterator))
 		{
-			PatchManager->RequestLandscapeUpdate();
+			return true;
 		}
 	}
+
+	return false;
 }
 
-void ULandscapePatchComponent::RequestLandscapeUpdate(bool bInUserTriggeredUpdate)
+void ULandscapePatchComponent::UpdateEditLayerFromDetailPanelLayerName()
 {
 	using namespace LandscapePatchComponentLocals;
 
-#if WITH_EDITOR
-	// Note that aside from the usual guard against doing things in the blueprint editor, the check of WorldType
-	// inside this call also prevents us from doing the request while cooking, where WorldType is set to Inactive. Otherwise
-	// we would issue a warning below while cooking when PatchManager is not saved (and not reset in the construction
-	// script).
-	if (!IsRealPatch())
+	if (!Landscape.IsValid())
 	{
 		return;
 	}
 
-	if (!PatchManager.IsValid())
+	const FLandscapeLayer* Layer = Landscape->GetLayerConst(*DetailPanelLayerName);
+	SetEditLayerGuid(Layer ? Layer->Guid : FGuid());
+}
+
+void ULandscapePatchComponent::FixBindings()
+{
+	Modify(false);
+
+	if (PatchManager.IsValid())
 	{
-		if (!bGaveMissingPatchManagerWarning)
+		if (PatchManager->GetOwningLandscape())
 		{
-			UE_LOG(LogLandscapePatch, Warning, TEXT("Patch does not have a valid patch manager. "
-				"Set the landscape or patch manager on the patch. (Package: %s, Actor: %s"), 
-				*GetNameSafe(GetPackage()), *GetNameSafe(GetAttachmentRootActor()));
-			bGaveMissingPatchManagerWarning = true;
+			ResetEditLayer();
+			Landscape = PatchManager->GetOwningLandscape();
+
+			if (!PatchManager->ContainsPatch(this) && IsPatchInWorld())
+			{
+				PatchManager->AddPatch(this);
+			}
+			return;
 		}
-		return;
+
+		// Otherwise clear out the invalid manager
+		ResetPatchManager();
 	}
 
-	bool bRequestUpdate = true;
-	if (!PatchManager->ContainsPatch(this))
+	if (EditLayerGuid.IsValid())
 	{
-		if (!bGaveNotInPatchManagerWarning)
+		if (BindToEditLayer(EditLayerGuid))
 		{
-			UE_LOG(LogLandscapePatch, Warning, TEXT("Patch's patch manager does not contain this patch. "
-				"Perhaps the manager was not saved? Reset the patch manager on the patch. (Package: %s, Actor: %s)"), 
-				*GetNameSafe(GetPackage()), *GetNameSafe(GetAttachmentRootActor()));
-			bGaveNotInPatchManagerWarning = true;
+			return;
 		}
-		bRequestUpdate = false;
+
+		// Otherwise, clear out the invalid guid
+		UE_LOG(LogLandscapePatch, Warning, TEXT("Could not find patch edit layer with guid: %s"), *EditLayerGuid.ToString());
+		ResetEditLayer();
 	}
-	if (!IsValid(PatchManager->GetOwningLandscape()))
+
+	// If we got here, we don't have an edit layer or patch manager
+	if (Landscape.IsValid())
 	{
-		if (!bGaveMissingLandscapeWarning)
+		if (BindToLandscape(Landscape.Get()))
 		{
-			UE_LOG(LogLandscapePatch, Warning, TEXT("Patch's patch manager does not have a valid owning "
-				"landscape. Perhaps the landscape was not saved? Reset the landscape on the manager. (Package: %s, Manager: %s)"), 
-				*GetNameSafe(GetPackage()), *GetNameSafe(PatchManager.Get()));
-			bGaveMissingLandscapeWarning = true;
+			return;
 		}
-		bRequestUpdate = false;
 	}
 
-	if(bRequestUpdate)
+	// At this point there's not much we can do to templates because we don't let them create
+	//  edit layers or bind to random landscapes (note: not actually sure that this function
+	//  would ever be called on templates... but we'll be safe)
+	if (!IsPatchInWorld())
 	{
-		ResetWarnings();
-		PatchManager->RequestLandscapeUpdate(bInUserTriggeredUpdate);
+		return;
+	}
+
+	// Try to bind to any landscape.
+	Landscape = nullptr;
+	if (IsEnabled() && !BindToAnyLandscape())
+	{
+		UE_LOG(LogLandscapePatch, Warning, TEXT("Unable to find landscape with edit layers enabled."));
+	}
+}
+
+void ULandscapePatchComponent::NotifyOfBoundLayerDeletion(ULandscapePatchEditLayer* LayerIn)
+{
+	// If we're getting this call, we are presumably bound through an edit layer, not patch manager
+	if (!ensure(!PatchManager.IsValid()))
+	{
+		ResetEditLayer();
+		return;
+	}
+
+	// This should update our cached pointer and displayed layer name
+	BindToEditLayer(EditLayerGuid);
+}
+#endif
+
+void ULandscapePatchComponent::SetPriority(double PriorityIn)
+{
+#if WITH_EDITOR
+	if (Priority == PriorityIn)
+	{
+		return;
+	}
+
+	Modify();
+	Priority = PriorityIn;
+	if (EditLayer.IsValid())
+	{
+		EditLayer->NotifyOfPriorityChange(this);
+		EditLayer->RequestLandscapeUpdate();
 	}
 #endif
 }
 
+void ULandscapePatchComponent::SetEditLayerGuid(const FGuid& GuidIn)
+{
 #if WITH_EDITOR
-void ULandscapePatchComponent::ResetWarnings()
+	Modify(false);
+	BindToEditLayer(GuidIn);
+	if (EditLayer.IsValid()) 
+	{
+		// We're not using a legacy manager
+		ResetPatchManager();
+	}
+	ResetWarnings();
+#endif
+}
+
+uint32 ULandscapePatchComponent::GetFullNameHash() const
 {
-	bGaveMissingPatchManagerWarning = false;
-	bGaveNotInPatchManagerWarning = false;
-	bGaveMissingLandscapeWarning = false;
+	// TODO: If needed, we can cache this value, but then we need to update it in
+	// PostLoad, creation, carry it over construction script reruns, and update it
+	// if the name changes.
+	return GetTypeHash(GetFullName());;
 }
-#endif
 
-bool ULandscapePatchComponent::IsRealPatch()
+// Provides a dropdown of available guids in the detail panel.
+const TArray<FString> ULandscapePatchComponent::GetLayerOptions()
 {
-	UWorld* World = GetWorld();
-	return !IsTemplate() && IsValid(this) && IsValid(World) && World->WorldType == EWorldType::Editor;
+	using namespace LandscapePatchComponentLocals;
+
+	TArray<FString> Options{ NullDetailPanelLayerName };
+
+#if WITH_EDITOR
+	// This is a bit of a separate concern, but update the displayed name if necessary. Usually this is not
+	//  needed because it is automatically updated, but it is possible for it to be stuck on a deleted layer 
+	//  in the blueprint editor, where our archetype is not registered to the layer and therefore doesn't get
+	//  the removal notification. 
+	if (!PatchManager.IsValid())
+	{
+		BindToEditLayer(EditLayerGuid);
+	}
+	else
+	{
+		// Uncertain whether we should try to fix a potential conflicting registration issue here (if
+		//  we both have an edit layer and a patch manager), but probably a good idea to, hence the clear
+		//  here. 
+		ResetEditLayer();
+	}
+
+	// Get the layer names from our landscape
+	if (Landscape.IsValid())
+	{
+		for (const FLandscapeLayer& Layer : Landscape->GetLayers())
+		{
+			if (Layer.EditLayer->IsA<ULandscapePatchEditLayer>())
+			{
+				Options.Add(Layer.Name.ToString());
+			}
+		}
+	}
+#endif
+
+	return Options;
 }
 
 FTransform ULandscapePatchComponent::GetLandscapeHeightmapCoordsToWorld() const
@@ -747,6 +1153,10 @@ FTransform ULandscapePatchComponent::GetLandscapeHeightmapCoordsToWorld() const
 	{
 		return PatchManager->GetHeightmapCoordsToWorld();
 	}
+	else if (Landscape.IsValid())
+	{
+		UE::Landscape::PatchUtil::GetHeightmapToWorld(Landscape->GetTransform());
+	}
 	return FTransform::Identity;
 }
 
diff --git a/Engine/Plugins/Experimental/LandscapePatch/Source/LandscapePatch/Private/LandscapePatchEditLayer.cpp b/Engine/Plugins/Experimental/LandscapePatch/Source/LandscapePatch/Private/LandscapePatchEditLayer.cpp
new file mode 100644
--- /dev/null
+++ b/Engine/Plugins/Experimental/LandscapePatch/Source/LandscapePatch/Private/LandscapePatchEditLayer.cpp
@@ -0,0 +1,442 @@
+// Copyright Epic Games, Inc. All Rights Reserved.
+
+#include "LandscapePatchEditLayer.h"
+
+#include "Algo/AnyOf.h"
+#include "Algo/Sort.h"
+#include "Landscape.h"
+#include "LandscapePatchLogging.h"
+#include "LandscapePatchUtil.h"
+#include "CoreGlobals.h" // UE::GetIsEditorLoadingPackage()
+
+#include "LandscapePatchComponent.h"
+
+#if WITH_EDITOR
+namespace LandscapePatchEditLayerLocals
+{
+	auto PatchSortPredicateRaw = [](ULandscapePatchComponent* A, ULandscapePatchComponent* B)
+	{
+		if (!ensure(A && B))
+		{
+			return false;
+		}
+
+		// If priorities are different, sort on priority
+		return A->GetPriority() != B->GetPriority() ? A->GetPriority() < B->GetPriority()
+
+			// If priorities are the same, use the full name hash. The comparison is meaningless, but gives
+			//  a deterministic ordering across runs, regardless of registration order.
+			: A->GetFullNameHash() != B->GetFullNameHash() ? A->GetFullNameHash() < B->GetFullNameHash()
+
+			// Hopefully we don't actually have to do full name string comparison, but that's the fallback.
+			: A->GetFullName() < B->GetFullName();
+	};
+
+	auto PatchSortPredicate = [](const TSoftObjectPtr<ULandscapePatchComponent>& ASoft, const TSoftObjectPtr<ULandscapePatchComponent>& BSoft)
+	{
+		return PatchSortPredicateRaw(ASoft.Get(), BSoft.Get());
+	};
+}
+
+void ULandscapePatchEditLayer::RegisterPatchForEditLayer(ULandscapePatchComponent* Patch)
+{
+	using namespace LandscapePatchEditLayerLocals;
+
+	if (!ShouldPatchBeIncludedInList(Patch))
+	{
+		return;
+	}
+
+	// See if we already have the patch
+	int32* FoundIndex = PatchToIndex.Find(Patch);
+	if (FoundIndex)
+	{
+		if (!ensureMsgf(RegisteredPatches.IsValidIndex(*FoundIndex) && RegisteredPatches[*FoundIndex] == Patch,
+			TEXT("LandscapePatchEditLayer: PatchToIndex is expected to match RegisteredPatces")))
+		{
+			bPatchListDirty = true;
+			FoundIndex = nullptr;
+		}
+	}
+
+	if (FoundIndex)
+	{
+		return;
+	}
+
+	Modify();
+
+	// See where this patch goes. If the list is up to date, we can get the actual insertion index
+	//  via a binary search. Otherwise we can just put on the end since we will resort anyway.
+	int32 InsertionIndex = RegisteredPatches.Num();
+	if (!bPatchListDirty)
+	{
+		InsertionIndex = GetInsertionIndex(Patch, bPatchListDirty);
+	}
+	
+	RegisteredPatches.Insert(Patch, InsertionIndex);
+
+	// Update index map for this and all patches forward
+	for (int32 i = InsertionIndex; i < RegisteredPatches.Num(); ++i)
+	{
+		PatchToIndex.Add(RegisteredPatches[i], i);
+	}
+
+	if (Patch->IsEnabled() && Patch->CanAffectLandscape())
+	{
+		RequestLandscapeUpdate();
+	}
+
+	UpdateHighestKnownPriority();
+}
+
+void ULandscapePatchEditLayer::NotifyOfPatchRemoval(ULandscapePatchComponent* Patch)
+{
+	using namespace LandscapePatchEditLayerLocals;
+
+	if (!ensure(Patch))
+	{
+		return;
+	}
+
+	if (!Patch->IsPatchInWorld() && ensure(!PatchToIndex.Contains(Patch)))
+	{
+		return;
+	}
+
+	// If we're being notified of removal, we expect that the patch doesn't point to our layer
+	//  or is otherwise legitimately not supposed to be in our list.
+	ensure(!ShouldPatchBeIncludedInList(Patch));
+
+	// See if we have this patch
+	int32 RemovedIndex = 0;
+	if (!PatchToIndex.RemoveAndCopyValue(Patch, RemovedIndex))
+	{
+		UE_LOG(LogLandscapePatch, Warning, TEXT("LandscapePatchEditLayer: Received NotifyOfPatchRemoval call for an unregistered patch."));
+		return;
+	}
+
+	ON_SCOPE_EXIT { RequestLandscapeUpdate(); };
+
+	if (!ensureMsgf(RegisteredPatches.IsValidIndex(RemovedIndex) && RegisteredPatches[RemovedIndex] == Patch, 
+		TEXT("LandscapePatchEditLayer: PatchToIndex is expected to match RegisteredPatces")))
+	{
+		bPatchListDirty = true;
+		return;
+	}
+
+	Modify();
+	RegisteredPatches.RemoveAt(RemovedIndex);
+
+	// Update forward indices
+	for (int32 i = RemovedIndex; i < RegisteredPatches.Num(); ++i)
+	{
+		PatchToIndex.Add(RegisteredPatches[i], i);
+	}
+
+	UpdateHighestKnownPriority();
+}
+
+void ULandscapePatchEditLayer::NotifyOfPriorityChange(ULandscapePatchComponent* Patch)
+{
+	using namespace LandscapePatchEditLayerLocals;
+
+	if (!ensure(Patch) || !Patch->IsPatchInWorld())
+	{
+		return;
+	}
+
+	TSoftObjectPtr<ULandscapePatchComponent> PatchSoftPtr = Patch;
+	int32* OriginalIndexPtr = PatchToIndex.Find(PatchSoftPtr);
+	if (!OriginalIndexPtr)
+	{
+		UE_LOG(LogLandscapePatch, Warning, TEXT("LandscapePatchEditLayer: Received NotifyOfPriorityChange call for an unregistered patch."));
+		return;
+	}
+
+	ON_SCOPE_EXIT{ RequestLandscapeUpdate(); };
+
+	if (bPatchListDirty)
+	{
+		// If patch list is dirty, we'll be resorting anyway, so no adjustment is needed right now.
+		return;
+	}
+
+	int32 OriginalIndex = *OriginalIndexPtr;
+	if (!ensureMsgf(RegisteredPatches[OriginalIndex] == Patch, TEXT("LandscapePatchEditLayer: PatchToIndex is expected to match RegisteredPatces")))
+	{
+		bPatchListDirty = true;
+		return;
+	}
+
+	// See if the patch is already in the proper place. Note that we only need to consider priority because only
+	//  priority changed (not the patch full name).
+	double Priority = Patch->GetPriority();
+	bool bPreviousPatchIsEqualOrLess = OriginalIndex == 0 
+		|| (RegisteredPatches[OriginalIndex - 1].IsValid() && RegisteredPatches[OriginalIndex - 1]->GetPriority() <= Priority);
+	bool bNextPatchIsEqualOrMore = OriginalIndex == RegisteredPatches.Num() - 1
+		|| (RegisteredPatches[OriginalIndex + 1].IsValid() && RegisteredPatches[OriginalIndex + 1]->GetPriority() >= Priority);
+	if (bPreviousPatchIsEqualOrLess && bNextPatchIsEqualOrMore)
+	{
+		return;
+	}
+
+	Modify();
+	RegisteredPatches.RemoveAt(OriginalIndex);
+	int32 InsertionIndex = GetInsertionIndex(Patch, bPatchListDirty);
+	RegisteredPatches.Insert(PatchSoftPtr, InsertionIndex);
+	
+	// Update all the indices that changed
+	int32 MinIndex = FMath::Min(OriginalIndex, InsertionIndex);
+	int32 MaxIndex = FMath::Max(OriginalIndex, InsertionIndex);
+	for (int32 i = MinIndex; i <= MaxIndex; ++i)
+	{
+		PatchToIndex.Add(RegisteredPatches[i], i);
+	}
+
+	UpdateHighestKnownPriority();
+}
+
+void ULandscapePatchEditLayer::UpdatePatchListIfDirty()
+{
+	if (bPatchListDirty)
+	{
+		UpdatePatchList();
+	}
+}
+
+void ULandscapePatchEditLayer::UpdateHighestKnownPriority()
+{
+	if (!bPatchListDirty)
+	{
+		HighestKnownPriority = RegisteredPatches.Num() > 0 ?
+			FMath::Max(RegisteredPatches.Last()->GetPriority(), PATCH_PRIORITY_BASE)
+			: PATCH_PRIORITY_BASE;
+	}
+}
+
+void ULandscapePatchEditLayer::UpdatePatchList()
+{
+	using namespace LandscapePatchEditLayerLocals;
+
+	// Filter out any patches that are no longer associated with this layer
+	RegisteredPatches.RemoveAll([this](TSoftObjectPtr<ULandscapePatchComponent> PatchSoft) 
+	{
+		ULandscapePatchComponent* Patch = PatchSoft.Get();
+		return !ShouldPatchBeIncludedInList(Patch);
+	});
+
+	// Sort by priority
+	Algo::Sort(RegisteredPatches, PatchSortPredicate);
+
+	// Update index lookup table
+	PatchToIndex.Reset();
+	for (int32 i = 0; i < RegisteredPatches.Num(); ++i)
+	{
+		PatchToIndex.Add(RegisteredPatches[i].Get(), i);
+	}
+
+	bPatchListDirty = false;
+	UpdateHighestKnownPriority();
+}
+
+double ULandscapePatchEditLayer::GetHighestPatchPriority()
+{
+	return HighestKnownPriority;
+}
+
+double ULandscapePatchEditLayer::ClaimHighestPatchPriority()
+{
+	HighestKnownPriority += 1;
+	return HighestKnownPriority;
+}
+#endif // WITH_EDITOR
+
+void ULandscapePatchEditLayer::GetRenderDependencies(TSet<UObject*>& OutDependencies) const
+{
+	Super::GetRenderDependencies(OutDependencies);
+
+#if WITH_EDITOR
+	for (const TSoftObjectPtr<ULandscapePatchComponent>& PatchSoft : RegisteredPatches)
+	{
+		ULandscapePatchComponent* Patch = PatchSoft.Get();
+		if (ShouldPatchBeIncludedInList(Patch))
+		{
+			Patch->GetRenderDependencies(OutDependencies);
+		}
+	}
+#endif
+}
+
+void ULandscapePatchEditLayer::OnLayerRemoved()
+{
+#if WITH_EDITOR
+	Modify();
+	// TODO: If we end up keeping this pointer, it should probably be reset in the base class implementation
+	// of OnLayerRemoved.
+	OwningLandscape.Reset();
+	for (TSoftObjectPtr<ULandscapePatchComponent> PatchSoft : RegisteredPatches)
+	{
+		ULandscapePatchComponent* Patch = PatchSoft.Get();
+		if (ShouldPatchBeIncludedInList(Patch))
+		{
+			Patch->NotifyOfBoundLayerDeletion(this);
+		}
+	}
+#endif
+}
+
+void ULandscapePatchEditLayer::OnLayerCreated(FLandscapeLayer& Layer)
+{
+	Super::OnLayerCreated(Layer);
+
+#if WITH_EDITOR
+	// TODO: Once the guid lives on the UObject, we won't need this.
+	EditLayerGuid = Layer.Guid;
+#endif
+}
+
+bool ULandscapePatchEditLayer::SupportsTargetType(ELandscapeToolTargetType InType) const
+{
+	return InType != ELandscapeToolTargetType::Invalid;
+}
+
+#if WITH_EDITOR
+
+void ULandscapePatchEditLayer::PostLoad()
+{
+	Super::PostLoad();
+
+	// Our layer should always have a guid that it got from OnLayerCreated. But if somehow that got lost,
+	// try to find the layer in our owning landscape.
+	if (!EditLayerGuid.IsValid() && ensure(OwningLandscape.IsValid()))
+	{
+		TArrayView<const FLandscapeLayer> Layers = OwningLandscape->GetLayers();
+		const FLandscapeLayer* LayerStruct = Layers.FindByPredicate([this](const FLandscapeLayer& Struct) { return Struct.EditLayer == this; });
+		EditLayerGuid = LayerStruct ? LayerStruct->Guid : FGuid();
+	}
+}
+
+void ULandscapePatchEditLayer::PostEditUndo()
+{
+	UpdatePatchList();
+}
+
+void ULandscapePatchEditLayer::RequestLandscapeUpdate(bool bInUserTriggered /* = false */)
+{
+	if (!OwningLandscape.IsValid())
+	{
+		return;
+	}
+
+	// TODO: Consider passing a parameter down to say when we're not updating height, only weights, when that is the case
+	OwningLandscape->RequestLayersContentUpdateForceAll(ELandscapeLayerUpdateMode::Update_All, bInUserTriggered || !UE::GetIsEditorLoadingPackage());
+}
+
+bool ULandscapePatchEditLayer::ShouldPatchBeIncludedInList(const ULandscapePatchComponent* Patch) const
+{
+	return Patch && Patch->IsPatchInWorld() && Patch->GetEditLayerGuid() == EditLayerGuid;
+}
+
+// Attempt binary search to find the insertion index for a patch. Returns end of array if an invalid patch is sampled.
+int32 ULandscapePatchEditLayer::GetInsertionIndex(ULandscapePatchComponent* Patch, bool& bFoundInvalidPatchOut) const
+{
+	using namespace LandscapePatchEditLayerLocals;
+
+	// This is a copy of BinarySearch.h::UpperBoundInternal, except that we have to check for validity
+	//  of our check values and exit early if we find an invalid one.
+
+	if (!ensure(Patch))
+	{
+		return INDEX_NONE;
+	}
+
+	if (bPatchListDirty)
+	{
+		bFoundInvalidPatchOut = true;
+		return RegisteredPatches.Num();
+	}
+
+	// Current start of sequence to check
+	int32 Start = 0;
+	// Size of sequence to check
+	int32 Size = RegisteredPatches.Num();
+
+	while (Size > 0)
+	{
+		const int32 LeftoverSize = Size % 2;
+		Size = Size / 2;
+
+		const int32 CheckIndex = Start + Size;
+		const int32 StartIfLess = CheckIndex + LeftoverSize;
+
+		ULandscapePatchComponent* CheckValue =  RegisteredPatches[CheckIndex].Get();
+		if (!CheckValue)
+		{
+			bFoundInvalidPatchOut = true;
+			return RegisteredPatches.Num();
+		}
+		Start = !PatchSortPredicateRaw(Patch, CheckValue) ? StartIfLess : Start;
+	}
+
+	return Start;
+}
+
+void ULandscapePatchEditLayer::InitializeAsBlueprintBrush(const FTransform& InLandscapeTransform, const FIntPoint& InLandscapeSize, const FIntPoint& InLandscapeRenderTargetSize)
+{
+	HeightmapCoordsToWorld = UE::Landscape::PatchUtil::GetHeightmapToWorld(InLandscapeTransform);
+}
+
+UTextureRenderTarget2D* ULandscapePatchEditLayer::RenderLayerAsBlueprintBrush(const FLandscapeBrushParameters& InParameters)
+{
+	UpdatePatchListIfDirty();
+
+	FLandscapeBrushParameters BrushParameters = InParameters;
+	for (TSoftObjectPtr<ULandscapePatchComponent>& PatchSoft : RegisteredPatches)
+	{
+		ULandscapePatchComponent* Patch = PatchSoft.Get();
+
+		// We shouldn't have any invalid patches because we should have been notified of dirtying or removals
+		if (!ensure(ShouldPatchBeIncludedInList(Patch)))
+		{
+			bPatchListDirty = true;
+			continue;
+		}
+
+		if (Patch->IsEnabled())
+		{
+			BrushParameters.CombinedResult = Patch->RenderLayer_Native(BrushParameters, HeightmapCoordsToWorld);
+		}
+	}
+
+	return BrushParameters.CombinedResult;
+}
+
+bool ULandscapePatchEditLayer::AffectsHeightmapAsBlueprintBrush() const
+{
+	return Algo::AnyOf(RegisteredPatches, [this](TSoftObjectPtr<ULandscapePatchComponent> PatchSoft)
+	{
+		ULandscapePatchComponent* Patch = PatchSoft.Get();
+		return ShouldPatchBeIncludedInList(Patch) && Patch->AffectsHeightmap();
+	});
+}
+
+bool ULandscapePatchEditLayer::AffectsWeightmapLayerAsBlueprintBrush(const FName& InLayerName) const
+{
+	return Algo::AnyOf(RegisteredPatches, [&InLayerName, this](TSoftObjectPtr<ULandscapePatchComponent> PatchSoft)
+	{ 
+		ULandscapePatchComponent* Patch = PatchSoft.Get();
+		return ShouldPatchBeIncludedInList(Patch) && Patch->AffectsWeightmapLayer(InLayerName);
+	});
+}
+
+bool ULandscapePatchEditLayer::AffectsVisibilityLayerAsBlueprintBrush() const
+{
+	return Algo::AnyOf(RegisteredPatches, [this](TSoftObjectPtr<ULandscapePatchComponent> PatchSoft)
+	{
+		ULandscapePatchComponent* Patch = PatchSoft.Get();
+		return ShouldPatchBeIncludedInList(Patch) && Patch->AffectsVisibilityLayer();
+	});
+}
+
+#endif // WITH_EDITOR
diff --git a/Engine/Plugins/Experimental/LandscapePatch/Source/LandscapePatch/Private/LandscapePatchManager.cpp b/Engine/Plugins/Experimental/LandscapePatch/Source/LandscapePatch/Private/LandscapePatchManager.cpp
--- a/Engine/Plugins/Experimental/LandscapePatch/Source/LandscapePatch/Private/LandscapePatchManager.cpp
+++ b/Engine/Plugins/Experimental/LandscapePatch/Source/LandscapePatch/Private/LandscapePatchManager.cpp
@@ -11,6 +11,7 @@
 #include "LandscapeDataAccess.h"
 #include "LandscapePatchComponent.h"
 #include "LandscapePatchLogging.h"
+#include "LandscapePatchUtil.h"
 #include "LandscapeModule.h"
 #include "LandscapeEditorServices.h"
 #include "Logging/MessageLog.h"
@@ -20,6 +21,15 @@
 #include "UObject/Package.h"
 #include "UObject/UObjectBaseUtility.h" // GetNameSafe
 #include "UObject/UObjectIterator.h"
+#include "Algo/AnyOf.h"
+
+#if WITH_EDITOR
+#include "Editor.h"
+#include "Editor/UnrealEdEngine.h"
+#include "LevelEditorSubsystem.h"
+#include "ScopedTransaction.h"
+#include "UnrealEdGlobals.h" //GUnrealEd
+#endif
 
 #include UE_INLINE_GENERATED_CPP_BY_NAME(LandscapePatchManager)
 
@@ -27,14 +37,19 @@
 
 namespace LandscapePatchManagerLocals
 {
+	const FText MigratePatchesTransactionName(LOCTEXT("MigratePatchesTransaction", "Migrate Patches"));
+
 #if WITH_EDITOR
-	// Another way to do the patch pointer fixup in addition to the map check.
-	FAutoConsoleCommand CCmdFixOwnedPatchBackPointers(
-		TEXT("LandscapePatch.PatchManager.FixOwnedPatchBackPointers"),
-		TEXT("For all patch managers, make it so that their managed patches point to them as their "
-			"patch manager. This is a more limited cleanup command than LandscapePatch.FixPatchManagerPointers."),
+	// Note: the priorities would become jumbled up if someone had multiple managers in the same edit layer. But
+	// this is an unexpected case that is not worth trying to handle differently.
+	FAutoConsoleCommand CCmdMigrateLegacyPatchListToPrioritySystem(
+		TEXT("LandscapePatch.MigrateLegacyListToPrioritySystem"),
+		TEXT("For all patch managers, make any patches in their patch list be directly bound to their edit layer, and "
+			"set the patch priorities according to their index."),
 		FConsoleCommandDelegate::CreateLambda([]() 
 	{
+		const FScopedTransaction Transaction(MigratePatchesTransactionName);
+
 		for (TObjectIterator<ALandscapePatchManager> It(
 			/*AdditionalExclusionFlags = */RF_ClassDefaultObject,
 			/*bIncludeDerivedClasses = */true,
@@ -54,11 +69,117 @@ namespace LandscapePatchManagerLocals
 
 			if (Manager)
 			{
-				Manager->FixOwnedPatchBackPointers();
+				Manager->MigrateToPrioritySystemAndDelete();
 			}
 		}
 	}));
 #endif // WITH_EDITOR
+
+	// Removes invalid patches from the list. This happens automatically when applying patches.
+	void FilterLegacyRegisteredPatches(TArray<TSoftObjectPtr<ULandscapePatchComponent>>& PatchComponents, 
+		TMap<TSoftObjectPtr<ULandscapePatchComponent>, int32>& PatchToIndex, const ALandscapePatchManager* ThisPatchManager)
+	{
+		// Used for removing invalid brushes. We remove from the index map immediately but then remove
+		// from the array and update other indices at the very end.
+		bool bHaveInvalidPatches = false;
+		int32 MinRemovedIndex = PatchComponents.Num();
+		auto RemoveComponentFromIndexMap = [&PatchToIndex, &MinRemovedIndex](TSoftObjectPtr<ULandscapePatchComponent>& Component)
+		{
+			int32 RemovedIndex = -1;
+			if (PatchToIndex.RemoveAndCopyValue(Component, RemovedIndex))
+			{
+				MinRemovedIndex = FMath::Min(MinRemovedIndex, RemovedIndex);
+			}
+		};
+
+		for (TSoftObjectPtr<ULandscapePatchComponent>& Component : PatchComponents)
+		{
+			if (Component.IsPending())
+			{
+				Component.LoadSynchronous();
+			}
+
+			if (Component.IsNull())
+			{
+				// Theoretically when components are marked for destruction, they should remove themselves from
+				// the patch manager in their OnComponentDestroyed call. However there seem to be ways to end up
+				// with destroyed patches not being removed, for instance through saving the manager but not the
+				// patch actor.
+				UE_LOG(LogLandscapePatch, Warning, TEXT("ALandscapePatchManager: Found an invalid patch in patch manager. It will be removed."));
+				RemoveComponentFromIndexMap(Component);
+				bHaveInvalidPatches = true;
+				continue;
+			}
+
+			if (!Component.IsValid())
+			{
+				// This means that IsPending() was true, but LoadSynchronous() failed, which we generally don't
+				// expect to happen. However, it can happen in some edge cases such as if you force delete a patch
+				// holder blueprint and don't save the patch manager afterward. Whatever the reason, this is likely
+				// a dead patch that actually needs removal.
+				UE_LOG(LogLandscapePatch, Warning, TEXT("ALandscapePatchManager: Found a pending patch pointer in patch manager that "
+					"turned out to be invalid. It will be removed."));
+				RemoveComponentFromIndexMap(Component);
+				bHaveInvalidPatches = true;
+				continue;
+			}
+
+			// Make sure the patch has this manager set as its patch manager.
+			PRAGMA_DISABLE_DEPRECATION_WARNINGS
+			if (Component->GetPatchManager() != ThisPatchManager)
+			PRAGMA_ENABLE_DEPRECATION_WARNINGS
+			{
+				UE_LOG(LogLandscapePatch, Warning, TEXT("ALandscapePatchManager: Found a patch whose patch manager is not set "
+						"to a patch manager that contains it. It will be removed."));
+				RemoveComponentFromIndexMap(Component);
+				bHaveInvalidPatches = true;
+				continue;
+			}
+
+			if (!Component->IsPatchInWorld())
+			{
+				UE_LOG(LogLandscapePatch, Warning, TEXT("ALandscapePatchManager: Found a non-world patch in patch manager. It will be removed."));
+				RemoveComponentFromIndexMap(Component);
+				bHaveInvalidPatches = true;
+				continue;
+			}
+		}
+
+		if (bHaveInvalidPatches)
+		{
+			PatchComponents.RemoveAll([ThisPatchManager](TSoftObjectPtr<ULandscapePatchComponent> Component) {
+				PRAGMA_DISABLE_DEPRECATION_WARNINGS
+				return !Component.IsValid() || Component->GetPatchManager() != ThisPatchManager || !Component->IsPatchInWorld();
+				PRAGMA_ENABLE_DEPRECATION_WARNINGS
+				});
+			// Update forward indices
+			for (int32 i = MinRemovedIndex; i < PatchComponents.Num(); ++i)
+			{
+				PatchToIndex.Add(PatchComponents[i], i);
+			}
+		}
+	}
+
+	/**
+	 * Returns true if any of the patches in the list pass the predicate.
+	 */
+	bool AnyOfPatchComponents(TArrayView<const TSoftObjectPtr<ULandscapePatchComponent>> LegacyRegisteredPatches,
+		TFunctionRef<bool(ULandscapePatchComponent&)> InPredicate)
+	{
+		return Algo::AnyOf(LegacyRegisteredPatches, [InPredicate](const TSoftObjectPtr<ULandscapePatchComponent>& InComponent)
+		{
+			if (InComponent.IsPending())
+			{
+				InComponent.LoadSynchronous();
+			}
+
+			if (InComponent.IsValid() && InPredicate(*InComponent.Get()))
+			{
+				return true;
+			}
+			return false;
+		});
+	}
 }
 
 // TODO: Not sure if using this kind of constructor is a proper thing to do vs some other hook...
@@ -76,141 +197,28 @@ void ALandscapePatchManager::Initialize_Native(const FTransform & InLandscapeTra
 	const FIntPoint& InLandscapeSize,
 	const FIntPoint& InLandscapeRenderTargetSize)
 {
-	// Get a transform from pixel coordinate in heightmap to world space coordinate. Note that we can't
-	// store the inverse directly because a FTransform can't properly represent a TRS inverse when the
-	// original TRS has non-uniform scaling).
-
-	// The pixel to landscape-space transform is unrotated, (S_p * x + T_p). The landscape to world
-	// transform gets applied on top of this: (R_l * S_l * (S_p * x + T_p)) + T_L. Collapsing this
-	// down to pixel to world TRS, we get: R_l * (S_l * S_p) * x + (R_l * S_l * T_p + T_L)
-
-	// To go from stored height value to unscaled height, we divide by 128 and subtract 256. We can get these
-	// values from the constants in LandscapeDataAccess.h (we distribute the multiplication by LANDSCAPE_ZSCALE
-	// so that translation happens after scaling like in TRS)
-	const double HEIGHTMAP_TO_OBJECT_HEIGHT_SCALE = LANDSCAPE_ZSCALE;
-	const double HEIGHTMAP_TO_OBJECT_HEIGHT_OFFSET = -LandscapeDataAccess::MidValue * LANDSCAPE_ZSCALE;
-
-	// S_p: the pixel coordinate scale is actually the same as xy object-space coordinates because one quad is 1 unit,
-	// so we only need to scale the height.
-	FVector3d PixelToObjectSpaceScale = FVector3d(
-		1,
-		1,
-		HEIGHTMAP_TO_OBJECT_HEIGHT_SCALE
-	);
-
-	// T_p: the center of the pixel
-	FVector3d PixelToObjectSpaceTranslate = FVector3d(
-		-0.5,
-		-0.5,
-		HEIGHTMAP_TO_OBJECT_HEIGHT_OFFSET
-	);
-
-	// S_l* S_p: composed scale
-	HeightmapCoordsToWorld.SetScale3D(InLandscapeTransform.GetScale3D() * PixelToObjectSpaceScale);
-
-	// R_l
-	HeightmapCoordsToWorld.SetRotation(InLandscapeTransform.GetRotation());
-
-	// R_l * S_l * T_p + T_L: composed translation
-	HeightmapCoordsToWorld.SetTranslation(InLandscapeTransform.TransformVector(PixelToObjectSpaceTranslate)
-		+ InLandscapeTransform.GetTranslation());
+	HeightmapCoordsToWorld = UE::Landscape::PatchUtil::GetHeightmapToWorld(InLandscapeTransform);
 }
 
+// Called in global merge to apply the patches
 UTextureRenderTarget2D* ALandscapePatchManager::RenderLayer_Native(const FLandscapeBrushParameters& InParameters)
 {
 	using namespace LandscapePatchManagerLocals;
 
-	// Used for removing invalid brushes. We remove from the index map immediately but then remove
-	// from the array and update other indices at the very end.
-	bool bHaveInvalidPatches = false;
-	int32 MinRemovedIndex = PatchComponents.Num();
-	auto RemoveComponentFromIndexMap = [this, &MinRemovedIndex](TSoftObjectPtr<ULandscapePatchComponent>& Component)
-	{
-		int32 RemovedIndex = -1;
-		if (PatchToIndex.RemoveAndCopyValue(Component, RemovedIndex))
-		{
-			MinRemovedIndex = FMath::Min(MinRemovedIndex, RemovedIndex);
-		}
-	};
+	// Note: We do not expect RenderLayer_Native to be called in the batched merge case.
+	// TODO: Check the cvar to make sure that is not the case?
 
+	FilterLegacyRegisteredPatches(PatchComponents, PatchToIndex, this);
 	FLandscapeBrushParameters BrushParameters = InParameters;
-
-	// TODO: There are many uncertainties in how we iterate across the height patches and have them
-	// apply themselves. For one thing we may want to pass around a render graph, in which case this
-	// loop will happen on the render thread somehow. For another, it's not yet determined what all
-	// of this will look like when we have the ability to render to just a subsection of the entire
-	// height map.
-	// So for now we do the simplest thing, and that is to have the height patches act as if they were
-	// independent brushes.
 	for (TSoftObjectPtr<ULandscapePatchComponent>& Component : PatchComponents)
 	{
-		if (Component.IsPending())
-		{
-			Component.LoadSynchronous();
-		}
-
-		if (Component.IsNull())
-		{
-			// Theoretically when components are marked for destruction, they should remove themselves from
-			// the patch manager in their OnComponentDestroyed call. However there seem to be ways to end up
-			// with destroyed patches not being removed, for instance through saving the manager but not the
-			// patch actor.
-			UE_LOG(LogLandscapePatch, Warning, TEXT("ALandscapePatchManager: Found an invalid patch in patch manager. It will be removed."));
-			RemoveComponentFromIndexMap(Component);
-			bHaveInvalidPatches = true;
-			continue;
-		}
-
-		if (!Component.IsValid())
-		{
-			// This means that IsPending() was true, but LoadSynchronous() failed, which we generally don't
-			// expect to happen. However, it can happen in some edge cases such as if you force delete a patch
-			// holder blueprint and don't save the patch manager afterward. Whatever the reason, this is likely
-			// a dead patch that actually needs removal.
-			UE_LOG(LogLandscapePatch, Warning, TEXT("ALandscapePatchManager: Found a pending patch pointer in patch manager that "
-				"turned out to be invalid. It will be removed."));
-			RemoveComponentFromIndexMap(Component);
-			Component = nullptr;
-			bHaveInvalidPatches = true;
-			continue;
-		}
-
 		if (!Component->IsEnabled())
 		{
 			// Skip disabled patches
 			continue;
 		}
 
-		// Make sure the patch has this manager set as its patch manager.
-		if (Component->GetPatchManager() != this)
-		{
-#if WITH_EDITORONLY_DATA
-			if (!bIssuedPatchOwnershipWarning)
-			{
-				bIssuedPatchOwnershipWarning = true;
-				UE_LOG(LogLandscapePatch, Warning, TEXT("ALandscapePatchManager: Found at least one patch whose patch manager is not "
-					"set to a patch manager that contains it. This should be fixed an the patches resaved (do this through the map check "
-					"action, or via LandscapePatch.PatchManager.FixOwnedPatchBackPointers console command.)"));
-			}
-#endif
-
-			// The patches need a proper back pointer to get parameters for their render call.
-			continue;
-		}
-
-		BrushParameters.CombinedResult = Component->RenderLayer_Native(BrushParameters); 
-	}
-
-	if (bHaveInvalidPatches)
-	{
-		PatchComponents.RemoveAll([](TSoftObjectPtr<ULandscapePatchComponent> Component) {
-			return Component.IsNull();
-		});
-		// Update forward indices
-		for (int32 i = MinRemovedIndex; i < PatchComponents.Num(); ++i)
-		{
-			PatchToIndex.Add(PatchComponents[i], i);
-		}
+		BrushParameters.CombinedResult = Component->RenderLayer_Native(BrushParameters, GetHeightmapCoordsToWorld()); 
 	}
 
 	return BrushParameters.CombinedResult;
@@ -248,7 +256,12 @@ void ALandscapePatchManager::SetTargetLandscape(ALandscape* InTargetLandscape)
 
 			ILandscapeModule& LandscapeModule = FModuleManager::GetModuleChecked<ILandscapeModule>("Landscape");
 			int32 PatchLayerIndex = LandscapeModule.GetLandscapeEditorServices()->GetOrCreateEditLayer(PatchLayerName, InTargetLandscape);
-			
+
+			if (!ensure(PatchLayerIndex >= 0))
+			{
+				return;
+			}
+
 			// Among other things, this will call SetOwningLandscape on us.
 			InTargetLandscape->AddBrushToLayer(PatchLayerIndex, this);
 
@@ -274,7 +287,7 @@ void ALandscapePatchManager::AddPatch(ULandscapePatchComponent* Patch)
 {
 	using namespace LandscapePatchManagerLocals;
 
-	if (Patch)
+	if (Patch && Patch->IsPatchInWorld())
 	{
 		if (!ContainsPatch(Patch))
 		{
@@ -284,18 +297,20 @@ void ALandscapePatchManager::AddPatch(ULandscapePatchComponent* Patch)
 			PatchToIndex.Add(PatchSoftPtr, PatchComponents.Num()-1);
 		}
 
+		PRAGMA_DISABLE_DEPRECATION_WARNINGS
 		if (Patch->GetPatchManager() != this)
 		{
 			UE_LOG(LogLandscapePatch, Warning, TEXT("ALandscapePatchManager::AddPatch: Added patch does not have this manager set "
-				"as its manager. Patches should typically be added to managers by setting the manager on the patch. "
+				"as its manager. Patches are typically added to managers by setting the manager on the patch. "
 				"(Package: % s, Actor : % s)"), *GetNameSafe(Patch->GetPackage()), *GetNameSafe(Patch->GetAttachmentRootActor()));
 		}
+		PRAGMA_ENABLE_DEPRECATION_WARNINGS
 
 		// No need to update if the patch is disabled. Important to avoid needlessly updating while dragging a blueprint with
 		// a disabled patch (since construction scripts constantly add and remove).
 		if (Patch->IsEnabled())
 		{
-			RequestLandscapeUpdate();
+			RequestLandscapeUpdate(!UE::GetIsEditorLoadingPackage());
 		}
 	}
 }
@@ -326,7 +341,7 @@ bool ALandscapePatchManager::RemovePatch(ULandscapePatchComponent* Patch)
 		// a blueprint with a disabled patch (since construction scripts constantly add and remove).
 		if (bRemoved && Patch->IsEnabled())
 		{
-			RequestLandscapeUpdate();
+			RequestLandscapeUpdate(!UE::GetIsEditorLoadingPackage());
 		}
 	}
 
@@ -344,7 +359,7 @@ int32 ALandscapePatchManager::GetIndexOfPatch(const ULandscapePatchComponent* Pa
 
 void ALandscapePatchManager::MovePatchToIndex(ULandscapePatchComponent* Patch, int32 Index)
 {
-	if (!Patch || Index < 0)
+	if (!Patch || !Patch->IsPatchInWorld() || Index < 0)
 	{
 		return;
 	}
@@ -384,74 +399,128 @@ void ALandscapePatchManager::MovePatchToIndex(ULandscapePatchComponent* Patch, i
 }
 
 #if WITH_EDITOR
-void ALandscapePatchManager::FixOwnedPatchBackPointers()
+void ALandscapePatchManager::MigrateToPrioritySystemAndDelete()
 {
-	for (TSoftObjectPtr<ULandscapePatchComponent>& Component : PatchComponents)
+	if (!PatchComponents.IsEmpty())
 	{
-		if (Component.IsValid() && Component->GetPatchManager() != this)
+		Modify();
+
+		// Patches will remove themselves from PatchComponents as we go along, so we need to iterate
+		// a copy.
+		TArray<TSoftObjectPtr<ULandscapePatchComponent>> PatchListCopy;
+
+		// We call Modify on all the patches we'll be touching at the start, otherwise they will
+		// store incorrect indices for undo as they are removed.
+		for (TSoftObjectPtr<ULandscapePatchComponent> Patch : PatchComponents)
 		{
-			Component->SetPatchManager(this);
+			if (Patch.IsPending())
+			{
+				Patch.LoadSynchronous();
+			}
+
+			if (Patch.IsValid())
+			{
+				Patch->Modify();
+				PatchListCopy.Add(Patch);
+			}
 		}
+
+		double Priority = LEGACY_PATCH_PRIORITY_BASE;
+		double PriorityStep = 1.0 / FMath::Max(1, PatchComponents.Num());
+
+		for (TSoftObjectPtr<ULandscapePatchComponent> Patch : PatchListCopy)
+		{
+			Patch->SetPriority(Priority);
+			Priority += PriorityStep;
+
+			PRAGMA_DISABLE_DEPRECATION_WARNINGS
+			Patch->SetPatchManager(nullptr);
+			PRAGMA_ENABLE_DEPRECATION_WARNINGS
+
+			Patch->FixBindings();
+		}
+
+		PatchComponents.Empty();
+		RequestLandscapeUpdate();
 	}
+
+	// Important so that we remove ourselves from the landscape blueprint brush list
+	SetTargetLandscape(nullptr);
+
+	// This will warn the user if there are still references to the manager
+	UTypedElementSelectionSet* SelectionSet = GEditor->GetEditorSubsystem<ULevelEditorSubsystem>()->GetSelectionSet();
+	GUnrealEd->DeleteActors({this}, GetWorld(), SelectionSet);
 }
 
-void ALandscapePatchManager::MarkModifiedInConstructionScript()
+
+bool ALandscapePatchManager::AffectsHeightmap() const
 {
-	bDirtiedByConstructionScript = true;
+	using namespace LandscapePatchManagerLocals;
+
+	if (!CanAffectHeightmap())
+	{
+		return false;
+	}
+
+	return AnyOfPatchComponents(PatchComponents,
+		[](ULandscapePatchComponent& InComponent) { return InComponent.AffectsHeightmap(); });
 }
 
-void ALandscapePatchManager::MarkDirtyIfModifiedInConstructionScript()
+bool ALandscapePatchManager::AffectsWeightmap() const
 {
-	if (bDirtiedByConstructionScript)
+	using namespace LandscapePatchManagerLocals;
+
+	if (!CanAffectWeightmap())
 	{
-		MarkPackageDirty();
+		return false;
 	}
+
+	return AnyOfPatchComponents(PatchComponents,
+		[](ULandscapePatchComponent& InComponent) { return InComponent.AffectsWeightmap(); });
+		
 }
 
 bool ALandscapePatchManager::AffectsWeightmapLayer(const FName& InLayerName) const
 {
+	using namespace LandscapePatchManagerLocals;
+
 	if (!CanAffectWeightmap())
 	{
 		return false;
 	}
 	
-	for (const TSoftObjectPtr<ULandscapePatchComponent>& Component : PatchComponents)
-	{
-		if (Component.IsPending())
-		{
-			Component.LoadSynchronous();
-		}
-
-		if (Component.IsValid() && Component->IsEnabled() && Component->AffectsWeightmapLayer(InLayerName))
-		{
-			return true;
-		}
-	}
-
-	return false;
+	return AnyOfPatchComponents(PatchComponents,
+		[InLayerName](ULandscapePatchComponent& InComponent) { return InComponent.AffectsWeightmapLayer(InLayerName); });
 }
 
 bool ALandscapePatchManager::AffectsVisibilityLayer() const
 {
+	using namespace LandscapePatchManagerLocals;
+
 	if (!CanAffectVisibilityLayer())
 	{
 		return false;
 	}
 
-	for (const TSoftObjectPtr<ULandscapePatchComponent>& Component : PatchComponents)
+	return AnyOfPatchComponents(PatchComponents,
+		[](ULandscapePatchComponent& InComponent) { return InComponent.AffectsVisibilityLayer(); });
+}
+
+void ALandscapePatchManager::GetRenderDependencies(TSet<UObject*>& OutDependencies)
+{
+	for (TSoftObjectPtr<ULandscapePatchComponent>& Component : PatchComponents)
 	{
 		if (Component.IsPending())
 		{
 			Component.LoadSynchronous();
 		}
 
-		if (Component.IsValid() && Component->IsEnabled() && Component->AffectsVisibilityLayer())
+		if (Component.IsValid())
 		{
-			return true;
+			Component->GetRenderDependencies(OutDependencies);
 		}
 	}
 
-	return false;
 }
 
 void ALandscapePatchManager::PostEditUndo()
@@ -459,14 +528,6 @@ void ALandscapePatchManager::PostEditUndo()
 	RequestLandscapeUpdate();
 }
 
-void ALandscapePatchManager::PreSave(FObjectPreSaveContext SaveContext)
-{
-	Super::PreSave(SaveContext);
-
-	// If we're saving, then we no longer have to worry about our dirtiness.
-	bDirtiedByConstructionScript = false;
-}
-
 void ALandscapePatchManager::SetOwningLandscape(ALandscape* InOwningLandscape)
 {
 	Super::SetOwningLandscape(InOwningLandscape);
@@ -518,47 +579,54 @@ void ALandscapePatchManager::CheckForErrors()
 		return Arguments;
 	};
 
-	// See if we're supposed to be dirty but aren't
-	if (bDirtiedByConstructionScript && !(GetPackage() && GetPackage()->IsDirty()))
-	{
-		FMessageLog("MapCheck").Warning()
-			->AddToken(FUObjectToken::Create(this))
-			->AddToken(FTextToken::Create(FText::Format(LOCTEXT("ManagerNeedsSaving", "Manager got modified in a construction "
-				"script rerun but has not been marked dirty. This could cause patch order to be unstable. "
-				"(Package: {Package}, Manager: {Actor})."), GetPackageAndActorArgs())))
-			->AddToken(FActionToken::Create(LOCTEXT("MarkDirtyButton", "Mark dirty"), FText(),
-				FOnActionTokenExecuted::CreateWeakLambda(this, [this]()
-		{
-			MarkPackageDirty();
-		})));
-	}
-
 	// See if we're holding on to any patches that don't have us as the owning patch manager
 	bool bHavePatchWithIncorrectManager = false;
 	for (TSoftObjectPtr<ULandscapePatchComponent>& Component : PatchComponents)
 	{
+		PRAGMA_DISABLE_DEPRECATION_WARNINGS
 		if (Component.IsValid() && Component->GetPatchManager() != this)
 		{
 			bHavePatchWithIncorrectManager = true;
 			break;
 		}
+		PRAGMA_ENABLE_DEPRECATION_WARNINGS
 	}
 
 	if (bHavePatchWithIncorrectManager)
 	{
-		FFormatNamedArguments Arguments;
-		Arguments.Add(TEXT("Manager"), FText::FromString(GetActorNameOrLabel()));
-		Arguments.Add(TEXT("Outer"), FText::FromString(GetPackage()->GetPathName()));
-
 		FMessageLog("MapCheck").Warning()
 			->AddToken(FUObjectToken::Create(this))
 			->AddToken(FTextToken::Create(FText::Format(LOCTEXT("PatchesHaveIncorrectManagerPointer", "Patch manager holds at "
-				"least one patch whose patch manager pointer is set incorrectly. These need to be fixed and resaved. "
+				"least one patch whose patch manager pointer is set incorrectly. These patches should be removed from the manager."
 				"(Package: {Package}, Manager: {Actor})."), GetPackageAndActorArgs())))
 			->AddToken(FActionToken::Create(LOCTEXT("FixPatchesButton", "Fix patches"), FText(),
 				FOnActionTokenExecuted::CreateWeakLambda(this, [this]() 
 		{
-			FixOwnedPatchBackPointers();
+			// Hard to say whether this should be in a transaction, or even be an action, because this happens
+			// automatically on the next landscape update... We'll stick with having it be user triggerable but
+			// not undoable.
+
+			PatchComponents.RemoveAll([this](TSoftObjectPtr<ULandscapePatchComponent> Component) {
+				PRAGMA_DISABLE_DEPRECATION_WARNINGS
+				return Component.IsValid() && Component->GetPatchManager() != this;
+				PRAGMA_ENABLE_DEPRECATION_WARNINGS
+				});
+		})));
+	}
+
+	if (PatchComponents.Num() > 0)
+	{
+		FMessageLog("MapCheck").Warning()
+			->AddToken(FUObjectToken::Create(this))
+			->AddToken(FTextToken::Create(FText::Format(LOCTEXT("UsingLegacyPatchList", "The use of the patch manager to "
+				"determine patch ordering is deprecated. Patches should point to a specific edit layer via a guid and "
+				"use Priority for ordering. You can use LandscapePatch.MigrateLegacyListToPrioritySystem to "
+				"fix this. (Package: {Package}, Manager: {Actor})."), GetPackageAndActorArgs())))
+			->AddToken(FActionToken::Create(LOCTEXT("MigrateToGuidsButton", "Migrate to guid system"), FText(),
+				FOnActionTokenExecuted::CreateWeakLambda(this, [this]() 
+		{
+			const FScopedTransaction Transaction(MigratePatchesTransactionName);
+			MigrateToPrioritySystemAndDelete();
 		})));
 	}
 }
diff --git a/Engine/Plugins/Experimental/LandscapePatch/Source/LandscapePatch/Private/LandscapePatchUtil.cpp b/Engine/Plugins/Experimental/LandscapePatch/Source/LandscapePatch/Private/LandscapePatchUtil.cpp
--- a/Engine/Plugins/Experimental/LandscapePatch/Source/LandscapePatch/Private/LandscapePatchUtil.cpp
+++ b/Engine/Plugins/Experimental/LandscapePatch/Source/LandscapePatch/Private/LandscapePatchUtil.cpp
@@ -2,6 +2,7 @@
 
 #include "LandscapePatchUtil.h"
 
+#include "LandscapeDataAccess.h" // LANDSCAPE_ZSCALE, MidValue
 #include "LandscapeTexturePatchPS.h" // FSimpleTextureCopyPS
 #include "RenderGraphBuilder.h"
 #include "TextureResource.h"
@@ -21,3 +22,49 @@ void UE::Landscape::PatchUtil::CopyTextureOnRenderThread(FRHICommandListImmediat
 
 	GraphBuilder.Execute();
 }
+
+FTransform UE::Landscape::PatchUtil::GetHeightmapToWorld(const FTransform& InLandscapeTransform)
+{
+	FTransform HeightmapCoordsToWorld;
+
+	// Get a transform from pixel coordinate in heightmap to world space coordinate. Note that we can't
+	// store the inverse directly because a FTransform can't properly represent a TRS inverse when the
+	// original TRS has non-uniform scaling).
+
+	// The pixel to landscape-space transform is unrotated, (S_p * x + T_p). The landscape to world
+	// transform gets applied on top of this: (R_l * S_l * (S_p * x + T_p)) + T_L. Collapsing this
+	// down to pixel to world TRS, we get: R_l * (S_l * S_p) * x + (R_l * S_l * T_p + T_L)
+
+	// To go from stored height value to unscaled height, we divide by 128 and subtract 256. We can get these
+	// values from the constants in LandscapeDataAccess.h (we distribute the multiplication by LANDSCAPE_ZSCALE
+	// so that translation happens after scaling like in TRS)
+	const double HEIGHTMAP_TO_OBJECT_HEIGHT_SCALE = LANDSCAPE_ZSCALE;
+	const double HEIGHTMAP_TO_OBJECT_HEIGHT_OFFSET = -LandscapeDataAccess::MidValue * LANDSCAPE_ZSCALE;
+
+	// S_p: the pixel coordinate scale is actually the same as xy object-space coordinates because one quad is 1 unit,
+	// so we only need to scale the height.
+	FVector3d PixelToObjectSpaceScale = FVector3d(
+		1,
+		1,
+		HEIGHTMAP_TO_OBJECT_HEIGHT_SCALE
+	);
+
+	// T_p: the center of the pixel
+	FVector3d PixelToObjectSpaceTranslate = FVector3d(
+		-0.5,
+		-0.5,
+		HEIGHTMAP_TO_OBJECT_HEIGHT_OFFSET
+	);
+
+	// S_l* S_p: composed scale
+	HeightmapCoordsToWorld.SetScale3D(InLandscapeTransform.GetScale3D() * PixelToObjectSpaceScale);
+
+	// R_l
+	HeightmapCoordsToWorld.SetRotation(InLandscapeTransform.GetRotation());
+
+	// R_l * S_l * T_p + T_L: composed translation
+	HeightmapCoordsToWorld.SetTranslation(InLandscapeTransform.TransformVector(PixelToObjectSpaceTranslate)
+		+ InLandscapeTransform.GetTranslation());
+
+	return HeightmapCoordsToWorld;
+}
\ No newline at end of file
diff --git a/Engine/Plugins/Experimental/LandscapePatch/Source/LandscapePatch/Private/LandscapeTexturePatch.cpp b/Engine/Plugins/Experimental/LandscapePatch/Source/LandscapePatch/Private/LandscapeTexturePatch.cpp
--- a/Engine/Plugins/Experimental/LandscapePatch/Source/LandscapePatch/Private/LandscapeTexturePatch.cpp
+++ b/Engine/Plugins/Experimental/LandscapePatch/Source/LandscapePatch/Private/LandscapeTexturePatch.cpp
@@ -77,15 +77,12 @@ namespace LandscapeTexturePatchLocals
 }
 
 #if WITH_EDITOR
-UTextureRenderTarget2D* ULandscapeTexturePatch::RenderLayer_Native(const FLandscapeBrushParameters& InParameters)
+UTextureRenderTarget2D* ULandscapeTexturePatch::RenderLayer_Native(const FLandscapeBrushParameters& InParameters, 
+	const FTransform& LandscapeHeightmapToWorld)
 {
 	using namespace UE::Landscape;
 
-	// If we're getting a RenderLayer_Native call, then we're inside some patch manager, and we
-	// expect our pointer to be properly initialized. This assumption could be violated if we
-	// were not consistent in saving both the manager and the patch. However the manager should
-	// be catching this case for us.
-	if (!ensure(PatchManager.IsValid()))
+	if (!IsPatchInWorld() || !IsEnabled())
 	{
 		return InParameters.CombinedResult;
 	}
@@ -99,12 +96,12 @@ UTextureRenderTarget2D* ULandscapeTexturePatch::RenderLayer_Native(const FLandsc
 		if (bReinitializeHeightOnNextRender)
 		{
 			bReinitializeHeightOnNextRender = false;
-			ReinitializeHeight(InParameters.CombinedResult);
+			ReinitializeHeight(InParameters.CombinedResult, LandscapeHeightmapToWorld);
 			return InParameters.CombinedResult;
 		}
 		else
 		{
-			return ApplyToHeightmap(InParameters.CombinedResult);
+			return ApplyToHeightmap(InParameters.CombinedResult, LandscapeHeightmapToWorld);
 		}
 	}
 	else
@@ -135,17 +132,18 @@ UTextureRenderTarget2D* ULandscapeTexturePatch::RenderLayer_Native(const FLandsc
 		if (WeightPatchInfo->bReinitializeOnNextRender)
 		{
 			WeightPatchInfo->bReinitializeOnNextRender = false;
-			ReinitializeWeightPatch(WeightPatchInfo, InParameters.CombinedResult);
+			ReinitializeWeightPatch(WeightPatchInfo, InParameters.CombinedResult, LandscapeHeightmapToWorld);
 			return InParameters.CombinedResult;
 		}
 		else
 		{
-			return ApplyToWeightmap(WeightPatchInfo, InParameters.CombinedResult);
+			return ApplyToWeightmap(WeightPatchInfo, InParameters.CombinedResult, LandscapeHeightmapToWorld);
 		}
 	}
 }
 
-UTextureRenderTarget2D* ULandscapeTexturePatch::ApplyToHeightmap(UTextureRenderTarget2D* InCombinedResult)
+UTextureRenderTarget2D* ULandscapeTexturePatch::ApplyToHeightmap(UTextureRenderTarget2D* InCombinedResult, 
+	const FTransform& LandscapeHeightmapToWorld)
 {
 	using namespace UE::Landscape;
 
@@ -189,7 +187,7 @@ UTextureRenderTarget2D* ULandscapeTexturePatch::ApplyToHeightmap(UTextureRenderT
 	// individually in the lambda below.
 	FApplyLandscapeTextureHeightPatchPS::FParameters ShaderParamsToCopy;
 	FIntRect DestinationBounds;
-	GetHeightShaderParams(FIntPoint(Patch->GetSizeX(), Patch->GetSizeY()), FIntPoint(InCombinedResult->SizeX, InCombinedResult->SizeY), ShaderParamsToCopy, DestinationBounds);
+	GetHeightShaderParams(LandscapeHeightmapToWorld, FIntPoint(Patch->GetSizeX(), Patch->GetSizeY()), FIntPoint(InCombinedResult->SizeX, InCombinedResult->SizeY), ShaderParamsToCopy, DestinationBounds);
 
 	if (DestinationBounds.IsEmpty())
 	{
@@ -237,7 +235,8 @@ UTextureRenderTarget2D* ULandscapeTexturePatch::ApplyToHeightmap(UTextureRenderT
 	return InCombinedResult;
 }
 
-UTextureRenderTarget2D* ULandscapeTexturePatch::ApplyToWeightmap(ULandscapeWeightPatchTextureInfo* PatchInfo, UTextureRenderTarget2D* InCombinedResult)
+UTextureRenderTarget2D* ULandscapeTexturePatch::ApplyToWeightmap(ULandscapeWeightPatchTextureInfo* PatchInfo, 
+	UTextureRenderTarget2D* InCombinedResult, const FTransform& LandscapeHeightmapToWorld)
 {
 	using namespace UE::Landscape;
 
@@ -286,7 +285,7 @@ UTextureRenderTarget2D* ULandscapeTexturePatch::ApplyToWeightmap(ULandscapeWeigh
 	FApplyLandscapeTextureWeightPatchPS::FParameters ShaderParamsToCopy;
 	FIntRect DestinationBounds;
 
-	GetWeightShaderParams(FIntPoint(Patch->GetSizeX(), Patch->GetSizeY()), FIntPoint(InCombinedResult->SizeX, InCombinedResult->SizeY), 
+	GetWeightShaderParams(LandscapeHeightmapToWorld, FIntPoint(Patch->GetSizeX(), Patch->GetSizeY()), FIntPoint(InCombinedResult->SizeX, InCombinedResult->SizeY),
 		PatchInfo, ShaderParamsToCopy, DestinationBounds);
 
 	if (DestinationBounds.IsEmpty())
@@ -335,7 +334,8 @@ UTextureRenderTarget2D* ULandscapeTexturePatch::ApplyToWeightmap(ULandscapeWeigh
 	return InCombinedResult;
 }
 
-void ULandscapeTexturePatch::GetCommonShaderParams(const FIntPoint& SourceResolutionIn, const FIntPoint& DestinationResolutionIn, 
+void ULandscapeTexturePatch::GetCommonShaderParams(const FTransform& LandscapeHeightmapToWorldIn,
+	const FIntPoint& SourceResolutionIn, const FIntPoint& DestinationResolutionIn, 
 	FTransform& PatchToWorldOut, FVector2f& PatchWorldDimensionsOut, FMatrix44f& HeightmapToPatchOut, FIntRect& DestinationBoundsOut, 
 	FVector2f& EdgeUVDeadBorderOut, float& FalloffWorldMarginOut) const
 {
@@ -348,8 +348,7 @@ void ULandscapeTexturePatch::GetCommonShaderParams(const FIntPoint& SourceResolu
 		FVector3d(FullPatchDimensions.X, FullPatchDimensions.Y, 1));
 	FMatrix44d PatchLocalToUVs = FromPatchUVToPatch.ToInverseMatrixWithScale();
 
-	FTransform LandscapeHeightmapToWorld = PatchManager->GetHeightmapCoordsToWorld();
-	FMatrix44d LandscapeToWorld = LandscapeHeightmapToWorld.ToMatrixWithScale();
+	FMatrix44d LandscapeToWorld = LandscapeHeightmapToWorldIn.ToMatrixWithScale();
 
 	FMatrix44d WorldToPatch = PatchToWorldOut.ToInverseMatrixWithScale();
 
@@ -361,11 +360,11 @@ void ULandscapeTexturePatch::GetCommonShaderParams(const FIntPoint& SourceResolu
 
 	// Get the output bounds, which are used to limit the amount of landscape pixels we have to process. 
 	// To get them, convert all of the corners into heightmap 2d coordinates and get the bounding box.
-	auto PatchUVToHeightmap2DCoordinates = [&PatchToWorldOut, &FromPatchUVToPatch, &LandscapeHeightmapToWorld](const FVector2f& UV)
+	auto PatchUVToHeightmap2DCoordinates = [&PatchToWorldOut, &FromPatchUVToPatch, &LandscapeHeightmapToWorldIn](const FVector2f& UV)
 	{
 		FVector WorldPosition = PatchToWorldOut.TransformPosition(
 			FromPatchUVToPatch.TransformPosition(FVector(UV.X, UV.Y, 0)));
-		FVector HeightmapCoordinates = LandscapeHeightmapToWorld.InverseTransformPosition(WorldPosition);
+		FVector HeightmapCoordinates = LandscapeHeightmapToWorldIn.InverseTransformPosition(WorldPosition);
 		return FVector2d(HeightmapCoordinates.X, HeightmapCoordinates.Y);
 	};
 	FBox2D FloatBounds(ForceInit);
@@ -391,7 +390,7 @@ void ULandscapeTexturePatch::GetCommonShaderParams(const FIntPoint& SourceResolu
 	FalloffWorldMarginOut = Falloff / FMath::Min(ComponentScale.X, ComponentScale.Y);
 }
 
-void ULandscapeTexturePatch::GetHeightShaderParams(
+void ULandscapeTexturePatch::GetHeightShaderParams(const FTransform& LandscapeHeightmapToWorldIn,
 	const FIntPoint& SourceResolutionIn, const FIntPoint& DestinationResolutionIn,
 	UE::Landscape::FApplyLandscapeTextureHeightPatchPS::FParameters& ParamsOut,
 	FIntRect& DestinationBoundsOut) const
@@ -399,7 +398,7 @@ void ULandscapeTexturePatch::GetHeightShaderParams(
 	using namespace UE::Landscape;
 
 	FTransform PatchToWorld;
-	GetCommonShaderParams(SourceResolutionIn, DestinationResolutionIn,
+	GetCommonShaderParams(LandscapeHeightmapToWorldIn, SourceResolutionIn, DestinationResolutionIn,
 		PatchToWorld, ParamsOut.InPatchWorldDimensions, ParamsOut.InHeightmapToPatch, 
 		DestinationBoundsOut, ParamsOut.InEdgeUVDeadBorder, ParamsOut.InFalloffWorldMargin);
 
@@ -428,13 +427,13 @@ void ULandscapeTexturePatch::GetHeightShaderParams(
 		break; // no offset necessary
 	case ELandscapeTextureHeightPatchZeroHeightMeaning::PatchZ:
 	{
-		FVector3d PatchOriginInHeightmapCoords = PatchManager->GetHeightmapCoordsToWorld().InverseTransformPosition(PatchToWorld.GetTranslation());
+		FVector3d PatchOriginInHeightmapCoords = LandscapeHeightmapToWorldIn.InverseTransformPosition(PatchToWorld.GetTranslation());
 		ParamsOut.InHeightOffset = PatchOriginInHeightmapCoords.Z - LandscapeDataAccess::MidValue;
 		break;
 	}
 	case ELandscapeTextureHeightPatchZeroHeightMeaning::WorldZero:
 	{
-		FVector3d WorldOriginInHeightmapCoords = PatchManager->GetHeightmapCoordsToWorld().InverseTransformPosition(FVector::ZeroVector);
+		FVector3d WorldOriginInHeightmapCoords = LandscapeHeightmapToWorldIn.InverseTransformPosition(FVector::ZeroVector);
 		ParamsOut.InHeightOffset = WorldOriginInHeightmapCoords.Z - LandscapeDataAccess::MidValue;
 		break;
 	}
@@ -460,7 +459,8 @@ void ULandscapeTexturePatch::GetHeightShaderParams(
 	ParamsOut.InFlags = static_cast<uint8>(Flags);
 }
 
-void ULandscapeTexturePatch::GetWeightShaderParams(const FIntPoint& SourceResolutionIn, 
+void ULandscapeTexturePatch::GetWeightShaderParams(
+	const FTransform& LandscapeHeightmapToWorldIn, const FIntPoint& SourceResolutionIn,
 	const FIntPoint& DestinationResolutionIn, const ULandscapeWeightPatchTextureInfo* WeightPatchInfo, 
 	UE::Landscape::FApplyLandscapeTextureWeightPatchPS::FParameters& ParamsOut, 
 	FIntRect& DestinationBoundsOut) const
@@ -468,7 +468,7 @@ void ULandscapeTexturePatch::GetWeightShaderParams(const FIntPoint& SourceResolu
 	using namespace UE::Landscape;
 
 	FTransform PatchToWorld;
-	GetCommonShaderParams(SourceResolutionIn, DestinationResolutionIn,
+	GetCommonShaderParams(LandscapeHeightmapToWorldIn, SourceResolutionIn, DestinationResolutionIn,
 		PatchToWorld, ParamsOut.InPatchWorldDimensions, ParamsOut.InWeightmapToPatch,
 		DestinationBoundsOut, ParamsOut.InEdgeUVDeadBorder, ParamsOut.InFalloffWorldMargin);
 
@@ -519,8 +519,7 @@ FLandscapeHeightPatchConvertToNativeParams ULandscapeTexturePatch::GetHeightConv
 
 #endif // WITH_EDITOR
 
-
-void ULandscapeTexturePatch::ReinitializeHeight()
+void ULandscapeTexturePatch::RequestReinitializeHeight()
 {
 #if WITH_EDITOR
 	if (!Super::IsEnabled())
@@ -529,12 +528,18 @@ void ULandscapeTexturePatch::ReinitializeHeight()
 		return;
 	}
 
-	if (!Landscape.IsValid() || !PatchManager.IsValid())
+	if (!Landscape.IsValid())
 	{
 		UE_LOG(LogLandscapePatch, Warning, TEXT("ULandscapeTexturePatch::Reinitialize: No associated landscape to initialize from."));
 		return;
 	}
 
+	if (!PatchManager.IsValid() && !GetBoundEditLayer())
+	{
+		UE_LOG(LogLandscapePatch, Warning, TEXT("ULandscapeTexturePatch::Reinitialize: Not bound to landscape (via edit layer)."));
+		return;
+	}
+
 	FVector2D DesiredResolution(FMath::Max(1, InitTextureSizeX), FMath::Max(1, InitTextureSizeY));
 	if (bBaseResolutionOffLandscape)
 	{
@@ -548,7 +553,7 @@ void ULandscapeTexturePatch::ReinitializeHeight()
 #endif // WITH_EDITOR
 }
 
-void ULandscapeTexturePatch::ReinitializeWeights()
+void ULandscapeTexturePatch::RequestReinitializeWeights()
 {
 #if WITH_EDITOR
 	if (!Super::IsEnabled())
@@ -557,12 +562,18 @@ void ULandscapeTexturePatch::ReinitializeWeights()
 		return;
 	}
 
-	if (!Landscape.IsValid() || !PatchManager.IsValid())
+	if (!Landscape.IsValid())
 	{
 		UE_LOG(LogLandscapePatch, Warning, TEXT("ULandscapeTexturePatch::Reinitialize: No associated landscape to initialize from."));
 		return;
 	}
 
+	if (!PatchManager.IsValid() && !GetBoundEditLayer())
+	{
+		UE_LOG(LogLandscapePatch, Warning, TEXT("ULandscapeTexturePatch::Reinitialize: Not bound to landscape (via edit layer)."));
+		return;
+	}
+
 	FVector2D DesiredResolution(FMath::Max(1, InitTextureSizeX), FMath::Max(1, InitTextureSizeY));
 	if (bBaseResolutionOffLandscape)
 	{
@@ -613,7 +624,7 @@ void ULandscapeTexturePatch::ReinitializeWeights()
 }
 
 #if WITH_EDITOR
-void ULandscapeTexturePatch::ReinitializeHeight(UTextureRenderTarget2D* InCombinedResult)
+void ULandscapeTexturePatch::ReinitializeHeight(UTextureRenderTarget2D* InCombinedResult, const FTransform& LandscapeHeightmapToWorld)
 {
 	if (HeightSourceMode == ELandscapeTexturePatchSourceMode::TextureAsset)
 	{
@@ -642,6 +653,10 @@ void ULandscapeTexturePatch::ReinitializeHeight(UTextureRenderTarget2D* InCombin
 		return;
 	}
 
+	SetUseAlphaChannelForHeight(false);
+	SetBlendMode(ELandscapeTexturePatchBlendMode::AlphaBlend);
+	ResetHeightRenderTargetFormat();
+
 	// The way we're going to do it is that we'll copy the packed values directly to a temporary render target, offset 
 	// them if needed (to undo whatever offsetting will happen during application), and store the result directly in the
 	// backing internal texture. Then we'll update the actual associated render target from the internal texture (if needed) so
@@ -661,7 +676,6 @@ void ULandscapeTexturePatch::ReinitializeHeight(UTextureRenderTarget2D* InCombin
 	double OffsetToApply = 0;
 	if (ZeroHeightMeaning != ELandscapeTextureHeightPatchZeroHeightMeaning::LandscapeZ)
 	{
-		FTransform LandscapeHeightmapToWorld = PatchManager->GetHeightmapCoordsToWorld();
 		double ZeroHeight = 0;
 		if (ZeroHeightMeaning == ELandscapeTextureHeightPatchZeroHeightMeaning::PatchZ)
 		{
@@ -674,7 +688,7 @@ void ULandscapeTexturePatch::ReinitializeHeight(UTextureRenderTarget2D* InCombin
 		OffsetToApply = LandscapeDataAccess::MidValue - ZeroHeight;
 	}
 
-	FMatrix44f PatchToSource = GetPatchToHeightmapUVs(TemporaryNativeHeightCopy->SizeX, TemporaryNativeHeightCopy->SizeY, InCombinedResult->SizeX, InCombinedResult->SizeY);
+	FMatrix44f PatchToSource = GetPatchToHeightmapUVs(LandscapeHeightmapToWorld, TemporaryNativeHeightCopy->SizeX, TemporaryNativeHeightCopy->SizeY, InCombinedResult->SizeX, InCombinedResult->SizeY);
 
 	ENQUEUE_RENDER_COMMAND(LandscapeTexturePatchReinitializeHeight)(
 		[Source = InCombinedResult->GetResource(), Destination = TemporaryNativeHeightCopy->GetResource(),
@@ -742,7 +756,8 @@ void ULandscapeTexturePatch::ReinitializeHeight(UTextureRenderTarget2D* InCombin
 	}
 }
 
-void ULandscapeTexturePatch::ReinitializeWeightPatch(ULandscapeWeightPatchTextureInfo* PatchInfo, UTextureRenderTarget2D* InCombinedResult)
+void ULandscapeTexturePatch::ReinitializeWeightPatch(ULandscapeWeightPatchTextureInfo* PatchInfo, 
+	UTextureRenderTarget2D* InCombinedResult, const FTransform& LandscapeHeightmapToWorld)
 {
 	using namespace LandscapeTexturePatchLocals;
 
@@ -779,6 +794,13 @@ void ULandscapeTexturePatch::ReinitializeWeightPatch(ULandscapeWeightPatchTextur
 		return;
 	}
 
+	PatchInfo->InternalData->SetUseAlphaChannel(false);
+	if (BlendMode != ELandscapeTexturePatchBlendMode::AlphaBlend)
+	{
+		PatchInfo->bOverrideBlendMode = true;
+		PatchInfo->OverrideBlendMode = ELandscapeTexturePatchBlendMode::AlphaBlend;
+	}
+
 	// We're going to copy directly to the associated render target. Make sure there is one for us to copy to.
 	PatchInfo->InternalData->SetUseInternalTextureOnly(false, false);
 	UTextureRenderTarget2D* RenderTarget = PatchInfo->InternalData->GetRenderTarget();
@@ -787,7 +809,8 @@ void ULandscapeTexturePatch::ReinitializeWeightPatch(ULandscapeWeightPatchTextur
 		return;
 	}
 
-	FMatrix44f PatchToSource = GetPatchToHeightmapUVs(RenderTarget->SizeX, RenderTarget->SizeY, InCombinedResult->SizeX, InCombinedResult->SizeY);
+	FMatrix44f PatchToSource = GetPatchToHeightmapUVs(LandscapeHeightmapToWorld, 
+		RenderTarget->SizeX, RenderTarget->SizeY, InCombinedResult->SizeX, InCombinedResult->SizeY);
 
 	ENQUEUE_RENDER_COMMAND(LandscapeTexturePatchReinitializeWeight)(
 		[Source = InCombinedResult->GetResource(), Destination = RenderTarget->GetResource(), &PatchToSource](FRHICommandListImmediate& RHICmdList)
@@ -815,7 +838,8 @@ void ULandscapeTexturePatch::ReinitializeWeightPatch(ULandscapeWeightPatchTextur
 	PatchInfo->InternalData->SetUseInternalTextureOnly(PatchInfo->SourceMode == ELandscapeTexturePatchSourceMode::InternalTexture, true);
 }
 
-FMatrix44f ULandscapeTexturePatch::GetPatchToHeightmapUVs(int32 PatchSizeX, int32 PatchSizeY, int32 HeightmapSizeX, int32 HeightmapSizeY) const
+FMatrix44f ULandscapeTexturePatch::GetPatchToHeightmapUVs(const FTransform& LandscapeHeightmapToWorld,
+	int32 PatchSizeX, int32 PatchSizeY, int32 HeightmapSizeX, int32 HeightmapSizeY) const
 {
 	FVector2D FullPatchDimensions = GetFullUnscaledWorldSize();
 
@@ -824,7 +848,6 @@ FMatrix44f ULandscapeTexturePatch::GetPatchToHeightmapUVs(int32 PatchSizeX, int3
 
 	FTransform PatchToWorld = GetPatchToWorldTransform();
 
-	FTransform LandscapeHeightmapToWorld = PatchManager->GetHeightmapCoordsToWorld();
 	FTransform LandscapeUVToWorld = LandscapeHeightmapToWorld;
 	LandscapeUVToWorld.MultiplyScale3D(FVector3d(HeightmapSizeX, HeightmapSizeY, 1));
 
@@ -835,61 +858,51 @@ FMatrix44f ULandscapeTexturePatch::GetPatchToHeightmapUVs(int32 PatchSizeX, int3
 	return (FMatrix44f)PatchToLandscapeUVTransposed.GetTransposed();
 }
 
-bool ULandscapeTexturePatch::AffectsWeightmapLayer(const FName& InLayerName) const
+bool ULandscapeTexturePatch::CanAffectHeightmap() const
 {
-	if (!IsEnabled())
-	{
-		return false;
-	}
+	return HeightSourceMode != ELandscapeTexturePatchSourceMode::None || bReinitializeHeightOnNextRender;
+}
 
-	for (const TObjectPtr<ULandscapeWeightPatchTextureInfo>& WeightPatch : WeightPatches)
-	{
-		if (IsValid(WeightPatch) && WeightPatch->WeightmapLayerName == InLayerName)
-		{
-			return true;
-		}
-	}
-	return false;
+bool ULandscapeTexturePatch::CanAffectWeightmap() const
+{
+	return Algo::AnyOf(WeightPatches, [](const TObjectPtr<ULandscapeWeightPatchTextureInfo>& InWeightPatch) 
+	{ 
+		return IsValid(InWeightPatch) && (InWeightPatch->SourceMode != ELandscapeTexturePatchSourceMode::None || InWeightPatch->bReinitializeOnNextRender);
+	});
 }
 
-bool ULandscapeTexturePatch::AffectsVisibilityLayer() const
+bool ULandscapeTexturePatch::CanAffectWeightmapLayer(const FName& InLayerName) const
 {
-	if (!IsEnabled())
+	return Algo::AnyOf(WeightPatches, [InLayerName](TObjectPtr<ULandscapeWeightPatchTextureInfo> InWeightPatch) 
 	{
-		return false;
-	}
+		return IsValid(InWeightPatch) && (InWeightPatch->SourceMode != ELandscapeTexturePatchSourceMode::None || InWeightPatch->bReinitializeOnNextRender)
+			&& (InWeightPatch->WeightmapLayerName == InLayerName);
+	}); 
+}
 
+bool ULandscapeTexturePatch::CanAffectVisibilityLayer() const
+{
 	return Algo::AnyOf(WeightPatches, [](const TObjectPtr<ULandscapeWeightPatchTextureInfo>& InWeightPatch) 
 	{ 
-		return IsValid(InWeightPatch) && InWeightPatch->bEditVisibilityLayer; 
+		return IsValid(InWeightPatch) && InWeightPatch->bEditVisibilityLayer 
+			&& (InWeightPatch->SourceMode != ELandscapeTexturePatchSourceMode::None || InWeightPatch->bReinitializeOnNextRender);
 	});
 }
 
-// We override IsEnabled to make the patch not request updates when all the source modes are "none"
-// (unless we need the update for reinitialization).
-bool ULandscapeTexturePatch::IsEnabled() const
+void ULandscapeTexturePatch::GetRenderDependencies(TSet<UObject*>& OutDependencies) const
 {
-	if (!Super::IsEnabled())
+	Super::GetRenderDependencies(OutDependencies);
+
+	if (HeightSourceMode == ELandscapeTexturePatchSourceMode::InternalTexture
+		&& HeightInternalData && HeightInternalData->GetInternalTexture())
 	{
-		return false;
+		OutDependencies.Add(HeightInternalData->GetInternalTexture());
 	}
-	if (HeightSourceMode != ELandscapeTexturePatchSourceMode::None || bReinitializeHeightOnNextRender)
+	else if (HeightSourceMode == ELandscapeTexturePatchSourceMode::TextureAsset 
+		&& HeightTextureAsset)
 	{
-		return true;
+		OutDependencies.Add(HeightTextureAsset);
 	}
-	for (const TObjectPtr<ULandscapeWeightPatchTextureInfo>& WeightPatch : WeightPatches)
-	{
-		if (!IsValid(WeightPatch))
-		{
-			continue;
-		}
-		if (WeightPatch->SourceMode != ELandscapeTexturePatchSourceMode::None || WeightPatch->bReinitializeOnNextRender)
-		{
-			return true;
-		}
-	}
-	// If we got to here, we are enabled, but all of the contained data had a source mode of "None"
-	return false;
 }
 
 void ULandscapeTexturePatch::MakeSureInternalDataIsAllocated()
@@ -921,7 +934,8 @@ void ULandscapeTexturePatch::MakeSureInternalDataIsAllocated()
 		}
 	}
 }
-#endif // WITH_EDITOR
+
+#endif
 
 void ULandscapeTexturePatch::SnapToLandscape()
 {
@@ -1124,6 +1138,17 @@ void ULandscapeTexturePatch::PostEditChangeProperty(FPropertyChangedEvent& Prope
 	{
 		if (PropertyChangedEvent.Property->GetFName() == GET_MEMBER_NAME_CHECKED(ULandscapeTexturePatch, DetailPanelHeightSourceMode))
 		{
+			// When changing source mode in the detail panel to a render target, we need to know the format to use, particularly 
+			// whether we need an alpha channel
+			if ((DetailPanelHeightSourceMode == ELandscapeTexturePatchSourceMode::TextureBackedRenderTarget
+					// This also affects an internal texture if we're copying from a texture asset, because we copy through render target
+					|| DetailPanelHeightSourceMode == ELandscapeTexturePatchSourceMode::InternalTexture)
+				// However we don't want to touch the format if we started with a render target source mode, because that would clear
+				// the render target before we can copy it to an internal texture (if that's what we're switching to).
+				&& HeightSourceMode != ELandscapeTexturePatchSourceMode::TextureBackedRenderTarget)
+			{
+				ResetHeightRenderTargetFormat();
+			}
 			SetHeightSourceMode(DetailPanelHeightSourceMode);
 		}
 		else if (PropertyChangedEvent.Property->GetFName() == GET_MEMBER_NAME_CHECKED(ULandscapeTexturePatch, HeightEncoding))
@@ -1178,7 +1203,7 @@ void ULandscapeWeightPatchTextureInfo::PreDuplicate(FObjectDuplicationParameters
 
 	Super::PreDuplicate(DupParams);
 
-	if (SourceMode == ELandscapeTexturePatchSourceMode::TextureBackedRenderTarget)
+	if (SourceMode == ELandscapeTexturePatchSourceMode::TextureBackedRenderTarget && InternalData)
 	{
 		InternalData->CopyToInternalTexture();
 	}
@@ -1320,6 +1345,7 @@ void ULandscapeTexturePatch::UpdateHeightConvertToNativeParamsIfNeeded()
 
 void ULandscapeTexturePatch::ResetHeightEncodingMode(ELandscapeTextureHeightPatchEncoding EncodingMode)
 {
+#if WITH_EDITOR
 	Modify();
 	HeightEncoding = EncodingMode;
 	if (EncodingMode == ELandscapeTextureHeightPatchEncoding::ZeroToOne)
@@ -1332,12 +1358,20 @@ void ULandscapeTexturePatch::ResetHeightEncodingMode(ELandscapeTextureHeightPatc
 		HeightEncodingSettings.ZeroInEncoding = 0;
 		HeightEncodingSettings.WorldSpaceEncodingScale = 1;
 	}
-	SetHeightRenderTargetFormat(EncodingMode == ELandscapeTextureHeightPatchEncoding::NativePackedHeight ?
-		ETextureRenderTargetFormat::RTF_RGBA8 : ETextureRenderTargetFormat::RTF_R32f);
+	ResetHeightRenderTargetFormat();
 
 	UpdateHeightConvertToNativeParamsIfNeeded();
+#endif
 }
 
+#if WITH_EDITOR
+void ULandscapeTexturePatch::ResetHeightRenderTargetFormat()
+{
+	SetHeightRenderTargetFormat(HeightEncoding == ELandscapeTextureHeightPatchEncoding::NativePackedHeight ? ETextureRenderTargetFormat::RTF_RGBA8
+		: bUseTextureAlphaForHeight ? ETextureRenderTargetFormat::RTF_RGBA32f : ETextureRenderTargetFormat::RTF_R32f);
+}
+#endif
+
 void ULandscapeTexturePatch::SetHeightEncodingSettings(const FLandscapeTexturePatchEncodingSettings& Settings)
 {
 	Modify();
diff --git a/Engine/Plugins/Experimental/LandscapePatch/Source/LandscapePatch/Public/LandscapeCircleHeightPatch.h b/Engine/Plugins/Experimental/LandscapePatch/Source/LandscapePatch/Public/LandscapeCircleHeightPatch.h
--- a/Engine/Plugins/Experimental/LandscapePatch/Source/LandscapePatch/Public/LandscapeCircleHeightPatch.h
+++ b/Engine/Plugins/Experimental/LandscapePatch/Source/LandscapePatch/Public/LandscapeCircleHeightPatch.h
@@ -19,16 +19,10 @@ class LANDSCAPEPATCH_API ULandscapeCircleHeightPatch : public ULandscapePatchCom
 
 public:
 
-	//TODO: For now the height patches have a similar interface to blueprint brushes, but this is likely to
-	// change. We are likely to pass a graph builder to the patches instead, along with an area that
-	// we want affected.
-	virtual void Initialize_Native(const FTransform& InLandscapeTransform, 
-		const FIntPoint& InLandscapeSize, 
-		const FIntPoint& InLandscapeRenderTargetSize) override;
+	virtual UTextureRenderTarget2D* RenderLayer_Native(const FLandscapeBrushParameters& InParameters, const FTransform& HeightmapToWorld) override;
 
-	virtual UTextureRenderTarget2D* RenderLayer_Native(const FLandscapeBrushParameters& InParameters) override;
-
-	virtual bool AffectsVisibilityLayer() const override { return bEditVisibility; }
+	virtual bool CanAffectHeightmap() const override { return !bEditVisibility; }
+	virtual bool CanAffectVisibilityLayer() const override { return bEditVisibility; }
 
 	// UActorComponent
 	virtual void OnComponentCreated() override;
diff --git a/Engine/Plugins/Experimental/LandscapePatch/Source/LandscapePatch/Public/LandscapePatchComponent.h b/Engine/Plugins/Experimental/LandscapePatch/Source/LandscapePatch/Public/LandscapePatchComponent.h
--- a/Engine/Plugins/Experimental/LandscapePatch/Source/LandscapePatch/Public/LandscapePatchComponent.h
+++ b/Engine/Plugins/Experimental/LandscapePatch/Source/LandscapePatch/Public/LandscapePatchComponent.h
@@ -6,6 +6,8 @@
 #include "Components/SceneComponent.h"
 #include "LandscapeBlueprintBrushBase.h"
 #include "LandscapeEditTypes.h"
+#include "LandscapePatchEditLayer.h" // PATCH_PRIORITY_BASE
+#include "UObject/WeakInterfacePtr.h"
 
 #include "LandscapePatchComponent.generated.h"
 
@@ -14,8 +16,27 @@ enum class ETeleportType : uint8;
 
 class ALandscape;
 class ALandscapePatchManager;
+struct FLandscapeLayer;
 class UTextureRenderTarget2D;
 
+UENUM()
+enum class ELandscapePatchPriorityInitialization : uint8
+{
+	// Initialize priority to highest currently known value, so that the new patch is on top of any existing
+	// patches. Note that the highest known priority could be out of date in between landscape updates if
+	// priorities change, so it is possible that adjustment will still be needed.
+	AcquireHighest,
+
+	// Do not change the default/archetype priority value. This is useful when using custom priority
+	// values as categories.
+	KeepOriginal,
+
+	// Increment the original priority by a small amount (0.01). This can be useful when copying a patch
+	// around multiple times, as it allows the new patches to be roughly in the same place in the 
+	// priority hierarchy while still being higher priority than the copied patch.
+	SmallIncrement,
+};
+
 /**
  * Base class for landscape patches: components that can be attached to meshes and moved around to make
  * the meshes affect the landscape around themselves.
@@ -30,14 +51,27 @@ class LANDSCAPEPATCH_API ULandscapePatchComponent : public USceneComponent
 public:
 	ULandscapePatchComponent(const FObjectInitializer& ObjectInitializer = FObjectInitializer::Get());
 
-	virtual void Initialize_Native(const FTransform& InLandscapeTransform,
-		const FIntPoint& InLandscapeSize,
-		const FIntPoint& InLandscapeRenderTargetSize) {}
+	// Called from global merge path
+	virtual UTextureRenderTarget2D* RenderLayer_Native(const FLandscapeBrushParameters& InParameters, 
+		const FTransform& HeightmapToWorld) { return InParameters.CombinedResult; }
+
+	// These determine whether the patch is configured correctly to affect height/weightmaps,
+	// ignoring whether it is currently enabled or not.
+	virtual bool CanAffectHeightmap() const { return false; }
+	virtual bool CanAffectWeightmap() const { return false; }
+	virtual bool CanAffectWeightmapLayer(const FName& InLayerName) const { return false; }
+	virtual bool CanAffectVisibilityLayer() const { return false; }
 
-	virtual UTextureRenderTarget2D* RenderLayer_Native(const FLandscapeBrushParameters& InParameters) { return InParameters.CombinedResult; }
+	bool CanAffectLandscape() const { return CanAffectHeightmap() || CanAffectWeightmap() || CanAffectVisibilityLayer(); }
 
-	virtual bool AffectsWeightmapLayer(const FName& InLayerName) const { return false; }
-	virtual bool AffectsVisibilityLayer() const { return false; }
+	// These compose IsEnabled with the appropriate CanAffect functions
+	bool AffectsHeightmap() const { return IsEnabled() && CanAffectHeightmap(); }
+	bool AffectsWeightmap() const { return IsEnabled() && CanAffectWeightmap(); }
+	bool AffectsWeightmapLayer(const FName& InLayerName) const { return IsEnabled() && CanAffectWeightmapLayer(InLayerName); }
+	bool AffectsVisibilityLayer() const { return IsEnabled() && CanAffectVisibilityLayer(); }
+
+	// Allows the patch to specify textures that need to be ready/compiled before applying the patch
+	virtual void GetRenderDependencies(TSet<UObject*>& OutDependencies) const {}
 
 	UFUNCTION(BlueprintCallable, Category = "LandscapePatch")
 	void RequestLandscapeUpdate(bool bInUserTriggeredUpdate = false);
@@ -55,36 +89,49 @@ public:
 	UFUNCTION(BlueprintCallable, Category = "LandscapePatch")
 	virtual bool IsEnabled() const { return bIsEnabled; }
 
+	UFUNCTION(BlueprintCallable, Category = "LandscapePatch")
+	double GetPriority() const { return Priority; }
+
+	UFUNCTION(BlueprintCallable, Category = "LandscapePatch")
+	void SetPriority(double PriorityIn);
+
+	FGuid GetEditLayerGuid() const { return EditLayerGuid; }
+	void SetEditLayerGuid(const FGuid& GuidIn);
+
+	/**
+	 * Returns the result of GetTypeHash(GetFullName()) for this component. Useful for quick deterministic
+	 * (though not lexical) ordering of same-priority patches. 
+	 */
+	uint32 GetFullNameHash() const;
+
 	UFUNCTION(BlueprintCallable, Category = "LandscapePatch")
 	FTransform GetLandscapeHeightmapCoordsToWorld() const;
 
 	UFUNCTION(BlueprintCallable, Category = "LandscapePatch")
 	virtual void SetLandscape(ALandscape* NewLandscape);
 
+	UE_DEPRECATED(5.5, "Patch manager use is deprecated. Patches should point to an edit layer using a Guid, and order themselves using priority.")
 	UFUNCTION(BlueprintCallable, Category = "LandscapePatch")
 	virtual void SetPatchManager(ALandscapePatchManager* NewPatchManager);
 
+	UE_DEPRECATED(5.5, "Patch manager use is deprecated. Patches should point to an edit layer using a Guid, and order themselves using priority.")
 	UFUNCTION(BlueprintCallable, Category = "LandscapePatch")
 	virtual ALandscapePatchManager* GetPatchManager() const;
 
-	void ApplyComponentInstanceData(struct FLandscapePatchComponentInstanceData* ComponentInstanceData, ECacheApplyPhase CacheApplyPhase);
+	// Verifies that we're not an archetype or trash or in an invalid world
+	bool IsPatchInWorld() const;
 
 	// For now we keep the patches largely editor-only, since we don't yet support runtime landscape editing.
 	// The above functions are also editor-only (and don't work at runtime), but can't be in WITH_EDITOR blocks
 	// so that they can be called from non-editor-only classes in editor contexts.
 #if WITH_EDITOR
 
-	/**
-	 * Verifies that the PatchManager and Landscape pointers are set up properly and alters them if needed. Public so
-	 * that it can be called from a console command.
-	 */
-	void SetUpPatchManagerData();
+	// Cleanup method that can be used to fix the patch's binding to a landscape. An arbitrary landscape
+	// will be bound to if no guid, patch manager, or landscape is set on the patch.
+	void FixBindings();
 
-	/**
-	 * Dirties the patch if it was modified in a construction script but was unable to mark itself
-	 * dirty. Meant to be used by cleanup commands.
-	 */
-	void MarkDirtyIfModifiedInConstructionScript();
+	// Allows the patch to update its UI and cached pointer
+	void NotifyOfBoundLayerDeletion(ULandscapePatchEditLayer* LayerIn);
 
 	// USceneComponent
 	virtual void OnUpdateTransform(EUpdateTransformFlags UpdateTransformFlags, ETeleportType Teleport) override;
@@ -101,22 +148,40 @@ public:
 	virtual void PostEditChangeProperty(FPropertyChangedEvent& PropertyChangedEvent) override;
 	virtual bool IsPostLoadThreadSafe() const override { return true; }
 	virtual void PostLoad() override;
-	virtual void PreSave(FObjectPreSaveContext SaveContext) override;
+	virtual void PostEditUndo() override;
 #endif
 	virtual bool IsEditorOnly() const override { return true; }
 	virtual bool NeedsLoadForClient() const override { return false; }
 	virtual bool NeedsLoadForServer() const override { return false; }
 protected:
-	/**
-	 * Move the patch to be the last processed patch in the current patch manager. This is a way to
-	 * reorder patches relative to each other.
-	 */
-	UFUNCTION(CallInEditor, Category = Initialization)
-	void MoveToTop();
+
+#if WITH_EDITOR
+	ULandscapePatchEditLayer* GetBoundEditLayer() { return EditLayer.Get(); }
+#endif
+
+	// Guid of the edit layer to which the patch is bound, if not using legacy patch manager list.
+	UPROPERTY()
+	FGuid EditLayerGuid;
+
+	/** How to initialize the patch priority when a patch is first created. */
+	UPROPERTY(EditAnywhere, Category = Settings, AdvancedDisplay)
+	ELandscapePatchPriorityInitialization PriorityInitialization = ELandscapePatchPriorityInitialization::AcquireHighest;
+
+	// Value that determines the patch ordering relative to other patches. Note that this is
+	// disregarded when using the legacy patch manager list (i.e. Patch Manager should be null).
+	UPROPERTY(EditAnywhere, Category = Settings)
+	double Priority = ULandscapePatchEditLayer::PATCH_PRIORITY_BASE;
 
 	UPROPERTY(EditAnywhere, Category = Settings)
 	TSoftObjectPtr<ALandscape> Landscape = nullptr;
 
+	/**
+	 * Deprecated: Patch manager that can hold the patch as part of its legacy patch list.
+	 * 
+	 * Instead, patches should now point to an edit layer via a guid and order themselves using their
+	 * priority value.
+	 */
+	//~ TODO: Deprecate this property, maybe if we have automatic conversion someday?
 	UPROPERTY(EditAnywhere, Category = Settings, AdvancedDisplay)
 	TSoftObjectPtr<ALandscapePatchManager> PatchManager = nullptr;
 
@@ -129,10 +194,7 @@ protected:
 	// Determines whether the height patch was made by copying a different height patch.
 	bool bWasCopy = false;
 
-	// This is true for existing height patches right after they are loaded, so that we can ignore
-	// the first OnRegister call. It remains false from the first OnRegiter call onward, even
-	// if the component is unregistered.
-	bool bLoadedButNotYetRegistered = false;
+#if WITH_EDITORONLY_DATA
 private:
 	// Starts as false and gets set to true in construction, so gets used to set bWasCopy
 	// by checking the indicator value at the start of construction.
@@ -143,26 +205,54 @@ private:
 	UPROPERTY()
 	TSoftObjectPtr<ALandscapePatchManager> PreviousPatchManager = nullptr;
 
-	// Verifies that we're not an archetype, trash, or not in a proper world
-	bool IsRealPatch();
+	//~ Name of edit layer displayed to user
+	/**
+	 * Name of the edit layer to which the patch is bound. Options are determined by the set Landscape pointer.
+	 */ 
+	UPROPERTY(Transient, EditAnywhere, Category = Settings, Meta = (DisplayName="Edit Layer", GetOptions = "GetLayerOptions"))
+	FString DetailPanelLayerName = TEXT("-Null-");
+
+	UPROPERTY(Transient, VisibleAnywhere, Category = Settings, AdvancedDisplay, Meta = (DisplayName = "Layer Guid"))
+	FString DetailPanelLayerGuid = EditLayerGuid.ToString();
+#endif
+
+	//~ Gives a drop down for edit layer names depending on the set landscape. Can't be editor only to be used in GetOptions
+	UFUNCTION()
+	const TArray<FString> GetLayerOptions();
 
 #if WITH_EDITOR
+	// Transient pointer to the edit layer that we are bound to using EditLayerGuid. Will
+	// be null if that layer is not of an appropriate type.
+	TWeakObjectPtr<ULandscapePatchEditLayer> EditLayer;
+
+	bool BindToEditLayer(const FGuid& Guid);
+	void ResetEditLayer();
+	bool BindToLandscape(ALandscape* LandscapeIn);
+	bool BindToAnyLandscape();
+	void UpdateEditLayerFromDetailPanelLayerName();
+
+	void ResetPatchManager();
+
+	// Tells us whether the patch is attached to an actor that is being dragged around as an editor preview
+	bool IsPatchPreview();
+
 	// Used to avoid spamming warning messages
-	bool bGaveMissingPatchManagerWarning = false;
+	bool bGaveCouldNotBindToEditLayerWarning = false;
+	bool bGaveMissingEditLayerGuidWarning = false;
+	bool bGaveMismatchedLandscapeWarning = false;
 	bool bGaveNotInPatchManagerWarning = false;
 	bool bGaveMissingLandscapeWarning = false;
 	void ResetWarnings();
-
-	// Automatic patch registration during on-load construction script reruns could end up modifying
-	// the patch without dirtying. This is dangerous as it can result in unstable patch ordering,
-	// so we want to detect this case.
-	bool bDirtiedByConstructionScript = false;
+	bool bInstanceDataApplied = false;
+	bool bDeferUpdateRequestUntilInstanceData = false;
+	
+	void ApplyComponentInstanceData(struct FLandscapePatchComponentInstanceData* ComponentInstanceData, ECacheApplyPhase CacheApplyPhase);
 #endif
 
 	friend struct FLandscapePatchComponentInstanceData;
 };
 
-/** Used to store some extra data during RerunConstructionScripts, namely the component's position in the patch manager. */
+/** Used to store some extra data during RerunConstructionScripts. */
 USTRUCT()
 struct FLandscapePatchComponentInstanceData : public FSceneComponentInstanceData
 {
@@ -180,8 +270,10 @@ struct FLandscapePatchComponentInstanceData : public FSceneComponentInstanceData
 
 	virtual void ApplyToComponent(UActorComponent* Component, const ECacheApplyPhase CacheApplyPhase) override
 	{
+#if WITH_EDITOR
 		Super::ApplyToComponent(Component, CacheApplyPhase);
 		CastChecked<ULandscapePatchComponent>(Component)->ApplyComponentInstanceData(this, CacheApplyPhase);
+#endif
 	}
 
 #if WITH_EDITORONLY_DATA
@@ -190,14 +282,21 @@ struct FLandscapePatchComponentInstanceData : public FSceneComponentInstanceData
 
 	UPROPERTY()
 	TSoftObjectPtr<ALandscapePatchManager> PatchManager = nullptr;
-	
 	UPROPERTY()
-	bool bDirtiedByConstructionScript = false;
+	FGuid EditLayerGuid;
+	// Priority needs to be carried over becuse our tweaks to it in OnComponentCreated cause it to not be
+	//  captured automatically, as it is detected as "UCS altered".
+	UPROPERTY()
+	double Priority = 0;
 
 	// Used so that we don't spam warning messages while rerunning construction scripts on a patch
 	// that triggers one of the warnings.
 	UPROPERTY()
-	bool bGaveMissingPatchManagerWarning = false;
+	bool bGaveCouldNotBindToEditLayerWarning = false;
+	UPROPERTY()
+	bool bGaveMismatchedLandscapeWarning = false;
+	UPROPERTY()
+	bool bGaveMissingEditLayerGuidWarning = false;
 	UPROPERTY()
 	bool bGaveNotInPatchManagerWarning = false;
 	UPROPERTY()
diff --git a/Engine/Plugins/Experimental/LandscapePatch/Source/LandscapePatch/Public/LandscapePatchEditLayer.h b/Engine/Plugins/Experimental/LandscapePatch/Source/LandscapePatch/Public/LandscapePatchEditLayer.h
new file mode 100644
--- /dev/null
+++ b/Engine/Plugins/Experimental/LandscapePatch/Source/LandscapePatch/Public/LandscapePatchEditLayer.h
@@ -0,0 +1,123 @@
+// Copyright Epic Games, Inc. All Rights Reserved.
+
+#pragma once
+
+#include "GlobalMergeLegacySupportUtil.h" // ILandscapeBrushRenderCallAdapter_GlobalMergeLegacySupport
+#include "LandscapeEditLayer.h"
+#include "Misc/Guid.h"
+
+#include "LandscapePatchEditLayer.generated.h"
+
+class ULandscapePatchComponent;
+
+/**
+ * Special edit layer used only for landscape patches. Patches point to the layer via a guid
+ *  and determine their ordering by their Priority values.
+ */
+UCLASS(MinimalAPI)
+class ULandscapePatchEditLayer : public ULandscapeEditLayerProcedural
+	, public ILandscapeBrushRenderCallAdapter_GlobalMergeLegacySupport
+{
+	GENERATED_BODY()
+
+public:
+#if WITH_EDITOR
+	/**
+	 * Must be called by patches on this layer to let the layer know that the patch is pointing
+	 *  to it via its layer guid.
+	 */
+	virtual void RegisterPatchForEditLayer(ULandscapePatchComponent* Patch);
+
+	/**
+	 * Should be called to notify the edit layer when a patch is destroyed or changes its edit layer guid
+	 *  away from this layer.
+	 */
+	virtual void NotifyOfPatchRemoval(ULandscapePatchComponent* Patch);
+
+	/**
+	 * Should be called to notify the edit layer when a patch priority changes.
+	 */
+	virtual void NotifyOfPriorityChange(ULandscapePatchComponent* Patch);
+
+	/**
+	 * Returns the highest known patch priority, or PATCH_PRIORITY_BASE if none are higher.
+	 */
+	virtual double GetHighestPatchPriority();
+
+	/**
+	 * Increments highest known patch priority and returns that value. Meant to allow an 
+	 *  already-registered patch to quickly become the highest priority patch.
+	 */
+	virtual double ClaimHighestPatchPriority();
+
+	void RequestLandscapeUpdate(bool bInUserTriggered = false);
+#endif
+
+	inline static const double PATCH_PRIORITY_BASE = 1000;
+
+	// ULandscapeEditLayerBase
+	virtual bool SupportsTargetType(ELandscapeToolTargetType InType) const override;
+	virtual bool NeedsPersistentTextures() const override { return false; }
+	virtual bool SupportsMultiple() const override { return true; }
+	virtual FString GetDefaultName() const override { return TEXT("Patches"); }
+	virtual void GetRenderDependencies(TSet<UObject*>& OutDependencies) const override;
+	virtual void OnLayerRemoved() override;
+	// TODO: Remove once our guid gets initialized somewhere else
+	virtual void OnLayerCreated(FLandscapeLayer& Layer);
+	
+#if WITH_EDITOR
+	// UObject
+	virtual void PostLoad() override;
+	virtual void PostEditUndo() override;
+
+	// ILandscapeBrushRenderCallAdapter_GlobalMergeLegacySupport
+	virtual bool AffectsHeightmapAsBlueprintBrush() const override;
+	virtual bool AffectsWeightmapLayerAsBlueprintBrush(const FName& InLayerName) const override;
+	virtual bool AffectsVisibilityLayerAsBlueprintBrush() const override;
+protected:
+	// ILandscapeBrushRenderCallAdapter_GlobalMergeLegacySupport
+	virtual void InitializeAsBlueprintBrush(const FTransform& InLandscapeTransform,
+		const FIntPoint& InLandscapeSize,
+		const FIntPoint& InLandscapeRenderTargetSize) override;
+	virtual UTextureRenderTarget2D* RenderLayerAsBlueprintBrush(const FLandscapeBrushParameters& InParameters) override;
+
+private:
+#endif
+#if WITH_EDITORONLY_DATA
+
+	// This is the transient list of patches that are bound to this edit layer, usually sorted by priority.
+	//  TSoftObjectPtr is used because it is robust across blueprint actor construction script reruns.
+	UPROPERTY(Transient)
+	TArray<TSoftObjectPtr<ULandscapePatchComponent>> RegisteredPatches;
+
+	// A helper structure for quick containment queries and updates. This should always be kept consistent with
+	//  RegisteredPatches.
+	UPROPERTY(Transient)
+	TMap<TSoftObjectPtr<ULandscapePatchComponent>, int32> PatchToIndex;
+
+	// When true, the entire patch list needs filtering and sorting. This is not meant to happen but it
+	//  is used as a safety in case some lack of notifications puts us in a situation where our patches
+	//  are not all valid and sorted.
+	bool bPatchListDirty = false;
+#endif
+#if WITH_EDITOR
+	void UpdatePatchList();
+	void UpdatePatchListIfDirty();
+	void UpdateHighestKnownPriority();
+
+	bool ShouldPatchBeIncludedInList(const ULandscapePatchComponent* Patch) const;
+	int32 GetInsertionIndex(ULandscapePatchComponent* Patch, bool& bFoundInvalidPatchOut) const;
+
+	// A tracker of the highest priority we've seen. When the patch list is not dirty, it will be the priority
+	//  of the last element, but this variable allows us to maintain the value past dirtying.
+	double HighestKnownPriority = PATCH_PRIORITY_BASE;
+
+	// Used in legacy global merge path to pass transform information from Initialize to RenderLayer
+	FTransform HeightmapCoordsToWorld;
+#endif
+#if WITH_EDITORONLY_DATA
+	// TODO: Remove this once guid is stored in the base edit layer UObject
+	UPROPERTY()
+	FGuid EditLayerGuid;
+#endif
+};
diff --git a/Engine/Plugins/Experimental/LandscapePatch/Source/LandscapePatch/Public/LandscapePatchManager.h b/Engine/Plugins/Experimental/LandscapePatch/Source/LandscapePatch/Public/LandscapePatchManager.h
--- a/Engine/Plugins/Experimental/LandscapePatch/Source/LandscapePatch/Public/LandscapePatchManager.h
+++ b/Engine/Plugins/Experimental/LandscapePatch/Source/LandscapePatch/Public/LandscapePatchManager.h
@@ -5,6 +5,7 @@
 
 #include "LandscapeBlueprintBrushBase.h"
 #include "LandscapeEditTypes.h"
+#include "LandscapePatchEditLayer.h" // PATCH_PRIORITY_BASE
 
 #include "LandscapePatchManager.generated.h"
 
@@ -12,26 +13,23 @@ class ALandscape;
 class ULandscapePatchComponent;
 
 /**
- * Acts as the "blueprint brush" as far as the owning edit layer is concerned. In reality, has its contained
- * patches edit the height/weight maps.
+ * Actor used in legacy landscape patch handling where a manager keeps a serialized list
+ * of patches that determines their priority. This approach is deprecated- patches now
+ * point to a special landscape patch edit layer via a guid, and determine their ordering
+ * relative to each other using a priority value.
  */
-//~ The alternative to this approach is to have the individual patches act as independent brushes, which
-//~ we currently don't want to do because we think it will clutter the brush interface and may lose opportunities
-//~ for optimization... 
 UCLASS()
-class LANDSCAPEPATCH_API ALandscapePatchManager : public ALandscapeBlueprintBrushBase
+class ALandscapePatchManager : public ALandscapeBlueprintBrushBase
 {
 	GENERATED_BODY()
 
 public:
-
 	ALandscapePatchManager(const FObjectInitializer& ObjectInitializer = FObjectInitializer::Get());
 
-	// These get called by the landscape system to apply the patches to the height/weight maps.
+	// These get called in the global merge path.
 	virtual void Initialize_Native(const FTransform& InLandscapeTransform,
 		const FIntPoint& InLandscapeSize,
 		const FIntPoint& InLandscapeRenderTargetSize) override;
-
 	virtual UTextureRenderTarget2D* RenderLayer_Native(const FLandscapeBrushParameters& InParameters) override;
 
 	// Adds the brush to the given landscape, removing it from any previous one. This differs from SetOwningLandscape
@@ -71,27 +69,21 @@ public:
 	void MovePatchToIndex(ULandscapePatchComponent* Patch, int32 Index);
 
 #if WITH_EDITOR
-	/**
-	 * A helper cleanup method to fix things if something goes wrong in saving and owned patches do not have
-	 * the correct patch manager pointer back. Public so that it can be called from a console command.
-	 */
-	void FixOwnedPatchBackPointers();
-
-	/**
-	 * Marks that the patch manager was modified during a construction script rerun where it might 
-	 * not be able to mark itself dirty (if it was done during loading).
-	 */
-	void MarkModifiedInConstructionScript();
 
 	/**
-	 * Dirties the manager if it was modified in a construction script but was unable to mark itself
-	 * dirty. Meant to be used by cleanup commands.
+	 * Move any patches from legacy patch list to being bound directly to an edit layer,
+	 * and delete the patch manager. This will cause a popup to the user if there is still
+	 * a dangling reference to the manager (there shouldn't be).
 	 */
-	void MarkDirtyIfModifiedInConstructionScript();
+	UFUNCTION(CallInEditor, Category = LandscapeManager, meta = (DisplayName = "MigrateToPrioritySystem"))
+	void MigrateToPrioritySystemAndDelete();
 
 	// ALandscapeBlueprintBrushBase
+	virtual bool AffectsHeightmap() const override;
+	virtual bool AffectsWeightmap() const override;
 	virtual bool AffectsWeightmapLayer(const FName& InLayerName) const override;
 	virtual bool AffectsVisibilityLayer() const override;
+	virtual void GetRenderDependencies(TSet<UObject*>& OutDependencies) override;
 	virtual void SetOwningLandscape(class ALandscape* InOwningLandscape) override;
 
 	// AActor
@@ -99,7 +91,6 @@ public:
 
 	// UObject
 	virtual void PostEditUndo() override;
-	virtual void PreSave(FObjectPreSaveContext SaveContext) override;
 	virtual void PostEditChangeProperty(FPropertyChangedEvent& PropertyChangedEvent) override;
 	virtual void PostLoad() override;
 #endif
@@ -107,11 +98,15 @@ public:
 	virtual bool NeedsLoadForClient() const override { return false; }
 	virtual bool NeedsLoadForServer() const override { return false; }
 
-protected:
+	// This is intentionally lower than PATCH_PRIORITY_BASE so that patches converted from a
+	// patch manager list are applied before other edit layer patches.
+	inline static const double LEGACY_PATCH_PRIORITY_BASE = ULandscapePatchEditLayer::PATCH_PRIORITY_BASE - 10;
+private:
 
 	UPROPERTY()
 	TArray<TSoftObjectPtr<ULandscapePatchComponent>> PatchComponents;
 
+	// Used in legacy paths to pass the transform information from Initialize_Native to RenderLayer_Native
 	UPROPERTY()
 	FTransform HeightmapCoordsToWorld;
 
@@ -123,14 +118,7 @@ protected:
 	UPROPERTY(EditAnywhere, Category = Landscape, Transient, meta = (DisplayName = "Landscape"))
 	TObjectPtr<ALandscape> DetailPanelLandscape = nullptr;
 
-private:
 	bool bIssuedPatchOwnershipWarning = false;
-
-	// The interaction of automatic patch registration and construction script reruns could end
-	// up modifying the manager during a load if things had to be fixed up, but the manager might
-	// not end up being marked dirty. This is dangerous as it can result in unstable patch ordering,
-	// so we want to detect this case.
-	bool bDirtiedByConstructionScript = false;
 #endif
 
 	// Transient table to speed up Contains and IndexOf queries, which are very slow for an array of TSoftObjectPtr's.
diff --git a/Engine/Plugins/Experimental/LandscapePatch/Source/LandscapePatch/Public/LandscapePatchUtil.h b/Engine/Plugins/Experimental/LandscapePatch/Source/LandscapePatch/Public/LandscapePatchUtil.h
--- a/Engine/Plugins/Experimental/LandscapePatch/Source/LandscapePatch/Public/LandscapePatchUtil.h
+++ b/Engine/Plugins/Experimental/LandscapePatch/Source/LandscapePatch/Public/LandscapePatchUtil.h
@@ -2,19 +2,19 @@
 
 #pragma once
 
-#include "HAL/Platform.h"
+#include "Math/MathFwd.h" // FTransform
 
 class FRHICommandListImmediate;
 class FTextureResource;
 
-namespace UE {
-namespace Landscape {
-namespace PatchUtil {
+namespace UE:: Landscape::PatchUtil 
+{
+	void CopyTextureOnRenderThread(FRHICommandListImmediate& RHICmdList, const FTextureResource& Source, FTextureResource& Destination);
 
-	void LANDSCAPEPATCH_API CopyTextureOnRenderThread(FRHICommandListImmediate& RHICmdList, const FTextureResource& Source, FTextureResource& Destination);
+	/**
+	 * Given a landscape transform, gives a transform from heightmap coordinates (where the Z value is the
+	 * two byte integer value stored as the height) to world coordinates.
+	 */
+	FTransform GetHeightmapToWorld(const FTransform& InLandscapeTransform);
 
-}}}//end UE::Landscape::PatchUtil
-
-#if UE_ENABLE_INCLUDE_ORDER_DEPRECATED_IN_5_2
-#include "CoreMinimal.h"
-#endif
+}//end UE::Landscape::PatchUtil
diff --git a/Engine/Plugins/Experimental/LandscapePatch/Source/LandscapePatch/Public/LandscapeTexturePatch.h b/Engine/Plugins/Experimental/LandscapePatch/Source/LandscapePatch/Public/LandscapeTexturePatch.h
--- a/Engine/Plugins/Experimental/LandscapePatch/Source/LandscapePatch/Public/LandscapeTexturePatch.h
+++ b/Engine/Plugins/Experimental/LandscapePatch/Source/LandscapePatch/Public/LandscapeTexturePatch.h
@@ -201,8 +201,6 @@ protected:
 	// TODO: We could support having different per-layer falloff modes and falloff amounts as well, as
 	// additional override members. But probably better to wait to see if that is actually desired.
 
-	// TODO: Like the similar flag for the height patch, this might not work once local merge works
-	// with landscape brushes...
 	bool bReinitializeOnNextRender = false;
 
 	void SetSourceMode(ELandscapeTexturePatchSourceMode NewMode);
@@ -223,12 +221,14 @@ class LANDSCAPEPATCH_API ULandscapeTexturePatch : public ULandscapePatchComponen
 public:
 
 #if WITH_EDITOR
-	UTextureRenderTarget2D* RenderLayer_Native(const FLandscapeBrushParameters& InParameters);
 
 	// ULandscapePatchComponent
-	virtual bool AffectsWeightmapLayer(const FName& InLayerName) const override;
-	virtual bool AffectsVisibilityLayer() const override;
-	virtual bool IsEnabled() const override;
+	UTextureRenderTarget2D* RenderLayer_Native(const FLandscapeBrushParameters& InParameters, const FTransform& HeightmapToWorld) override;
+	virtual bool CanAffectHeightmap() const override;
+	virtual bool CanAffectWeightmap() const override;
+	virtual bool CanAffectWeightmapLayer(const FName& InLayerName) const override;
+	virtual bool CanAffectVisibilityLayer() const override;
+	virtual void GetRenderDependencies(TSet<UObject*>& OutDependencies) const override;
 
 	// UActorComponent
 	virtual TStructOnScope<FActorComponentInstanceData> GetComponentInstanceData() const override;
@@ -319,8 +319,9 @@ public:
 	virtual ELandscapeTexturePatchSourceMode GetHeightSourceMode() const { return HeightSourceMode; }
 
 	/**
-	 * Changes source mode. When changing between sources, existing data is copied from one to the other
-	 * when possible.
+	 * Changes source mode. There are currently no API guarantees regarding the initialization of the
+	 * new source data. E.g. when first switching to use an internal render target, the data in that
+	 * render target may not be initialized.
 	 */
 	UFUNCTION(BlueprintCallable, Category = LandscapePatch)
 	virtual void SetHeightSourceMode(ELandscapeTexturePatchSourceMode NewMode);
@@ -560,11 +561,11 @@ protected:
 	/**
 	 * Given the current initialization settings, reinitialize the height patch.
 	 */
-	UFUNCTION(CallInEditor, Category = HeightPatch)
-	void ReinitializeHeight();
+	UFUNCTION(CallInEditor, Category = HeightPatch, meta = (DisplayName= "Reinitialize Height"))
+	void RequestReinitializeHeight();
 
-	UFUNCTION(CallInEditor, Category = WeightPatches)
-	void ReinitializeWeights();
+	UFUNCTION(CallInEditor, Category = WeightPatches, meta = (DisplayName = "Reinitialize Weights"))
+	void RequestReinitializeWeights();
 
 	bool bReinitializeHeightOnNextRender = false;
 
@@ -611,22 +612,27 @@ private:
 #if WITH_EDITOR
 	void TransitionHeightSourceModeInternal(ELandscapeTexturePatchSourceMode OldMode, ELandscapeTexturePatchSourceMode NewMode);
 	FLandscapeHeightPatchConvertToNativeParams GetHeightConvertToNativeParams() const;
-	UTextureRenderTarget2D* ApplyToHeightmap(UTextureRenderTarget2D* InCombinedResult);
-	UTextureRenderTarget2D* ApplyToWeightmap(ULandscapeWeightPatchTextureInfo* PatchInfo, UTextureRenderTarget2D* InCombinedResult);
+	UTextureRenderTarget2D* ApplyToHeightmap(UTextureRenderTarget2D* InCombinedResult, const FTransform& LandscapeHeightmapToWorld);
+	UTextureRenderTarget2D* ApplyToWeightmap(ULandscapeWeightPatchTextureInfo* PatchInfo, 
+		UTextureRenderTarget2D* InCombinedResult, const FTransform& LandscapeHeightmapToWorld);
 
-	void GetCommonShaderParams(const FIntPoint& SourceResolutionIn, const FIntPoint& DestinationResolutionIn, 
+	void GetCommonShaderParams(const FTransform& LandscapeHeightmapToWorldIn, 
+		const FIntPoint& SourceResolutionIn, const FIntPoint& DestinationResolutionIn,
 		FTransform& PatchToWorldOut, FVector2f& PatchWorldDimensionsOut, FMatrix44f& HeightmapToPatchOut, 
 		FIntRect& DestinationBoundsOut, FVector2f& EdgeUVDeadBorderOut, float& FalloffWorldMarginOut) const;
-	void GetHeightShaderParams(const FIntPoint& SourceResolutionIn, const FIntPoint& DestinationResolutionIn,
+	void GetHeightShaderParams(const FTransform& LandscapeHeightmapToWorldIn, 
+		const FIntPoint& SourceResolutionIn, const FIntPoint& DestinationResolutionIn,
 		UE::Landscape::FApplyLandscapeTextureHeightPatchPS::FParameters& ParamsOut, FIntRect& DestinationBoundsOut) const;
-	void GetWeightShaderParams(const FIntPoint& SourceResolutionIn, const FIntPoint& DestinationResolutionIn,
-		const ULandscapeWeightPatchTextureInfo* WeightPatchInfo, 
+	void GetWeightShaderParams(const FTransform& LandscapeHeightmapToWorldIn, const FIntPoint& SourceResolutionIn, 
+		const FIntPoint& DestinationResolutionIn, const ULandscapeWeightPatchTextureInfo* WeightPatchInfo, 
 		UE::Landscape::FApplyLandscapeTextureWeightPatchPS::FParameters& ParamsOut, FIntRect& DestinationBoundsOut) const;
-	FMatrix44f GetPatchToHeightmapUVs(int32 PatchSizeX, int32 PatchSizeY, int32 HeightmapSizeX, int32 HeightmapSizeY) const;
-	void ReinitializeHeight(UTextureRenderTarget2D* InCombinedResult);
-	void ReinitializeWeightPatch(ULandscapeWeightPatchTextureInfo* PatchInfo, UTextureRenderTarget2D* InCombinedResult);
+	FMatrix44f GetPatchToHeightmapUVs(const FTransform& LandscapeHeightmapToWorld, int32 PatchSizeX, int32 PatchSizeY, int32 HeightmapSizeX, int32 HeightmapSizeY) const;
+	void ReinitializeHeight(UTextureRenderTarget2D* InCombinedResult, const FTransform& LandscapeHeightmapToWorld);
+	void ReinitializeWeightPatch(ULandscapeWeightPatchTextureInfo* PatchInfo, UTextureRenderTarget2D* InCombinedResult, 
+		const FTransform& LandscapeHeightmapToWorld);
 
 	void MakeSureInternalDataIsAllocated();
+	void ResetHeightRenderTargetFormat();
 #endif // WITH_EDITOR
 
 	UPROPERTY(EditDefaultsOnly, Category = Settings)
diff --git a/Engine/Plugins/Experimental/LandscapePatch/Source/LandscapePatchEditorOnly/Private/LandscapeTexturePatchVisualizer.h b/Engine/Plugins/Experimental/LandscapePatch/Source/LandscapePatchEditorOnly/Private/LandscapeTexturePatchVisualizer.h
new file mode 100644
--- /dev/null
+++ b/Engine/Plugins/Experimental/LandscapePatch/Source/LandscapePatchEditorOnly/Private/LandscapeTexturePatchVisualizer.h
@@ -0,0 +1,19 @@
+// Copyright Epic Games, Inc. All Rights Reserved.
+
+#pragma once
+
+
+#include "ComponentVisualizer.h"
+
+/**
+ * Adds a rectangle denoting the area covered by a landscape texture patch when it is selected in editor.
+ */
+class LANDSCAPEPATCHEDITORONLY_API FLandscapeTexturePatchVisualizer : public FComponentVisualizer
+{
+private:
+	virtual void DrawVisualization(const UActorComponent* Component, const FSceneView* View, FPrimitiveDrawInterface* PDI) override;
+};
+
+#if UE_ENABLE_INCLUDE_ORDER_DEPRECATED_IN_5_2
+#include "CoreMinimal.h"
+#endif
diff --git a/Engine/Plugins/Experimental/LandscapePatch/Source/LandscapePatchEditorOnly/Public/LandscapeTexturePatchVisualizer.h b/Engine/Plugins/Experimental/LandscapePatch/Source/LandscapePatchEditorOnly/Public/LandscapeTexturePatchVisualizer.h
deleted file mode 100644
diff --git a/Engine/Source/Runtime/Landscape/Classes/Landscape.h b/Engine/Source/Runtime/Landscape/Classes/Landscape.h
--- a/Engine/Source/Runtime/Landscape/Classes/Landscape.h
+++ b/Engine/Source/Runtime/Landscape/Classes/Landscape.h
@@ -459,8 +459,7 @@ public:
 
 	LANDSCAPE_API void SetDisableRuntimeGrassMapGeneration(bool bInDisableRuntimeGrassMapGeneration);
 
-protected:
-	FName GenerateUniqueLayerName(FName InName = NAME_None) const;
+	LANDSCAPE_API FName GenerateUniqueLayerName(FName InName = NAME_None) const;
 
 private:
 	FLandscapeLayer* GetLayerInternal(int32 InLayerIndex);
@@ -475,6 +474,8 @@ private:
 	void MonitorShaderCompilation();
 	void MonitorLandscapeEdModeChanges();
 	
+	FLandscapeBrushParameters BuildLandscapeBrushParameters(bool bInIsHeightmapMerge, const FIntRect& InRenderAreaExtents, UTextureRenderTarget2D* InRenderTarget, FName InWeightmapLayerName = NAME_None);
+
 	int32 RegenerateLayersHeightmaps(const FUpdateLayersContentContext& InUpdateLayersContentContext);
 	int32 PerformLayersHeightmapsLocalMerge(const FUpdateLayersContentContext& InUpdateLayersContentContext, const FEditLayersHeightmapMergeParams& InMergeParams);
 	int32 PerformLayersHeightmapsGlobalMerge(const FUpdateLayersContentContext& InUpdateLayersContentContext, const FEditLayersHeightmapMergeParams& InMergeParams);
diff --git a/Engine/Source/Runtime/Landscape/Classes/LandscapeEditLayer.h b/Engine/Source/Runtime/Landscape/Classes/LandscapeEditLayer.h
--- a/Engine/Source/Runtime/Landscape/Classes/LandscapeEditLayer.h
+++ b/Engine/Source/Runtime/Landscape/Classes/LandscapeEditLayer.h
@@ -72,6 +72,11 @@ public:
 	};
 
 public:
+
+	// TODO: This might be removed once the guid is stored here and subclasses have a way to request landscape updates.
+	// Otherwise, it might be better made private and then we friend ALandscape?
+	void SetBackPointer(ALandscape* Landscape);
+
 	/**
 	 * @return true if the this edit layer has support for the target type (heightmap, weightmap, visibility)
 	 */
@@ -132,6 +137,19 @@ public:
 	//  settings, on property change, the event would be triggered and the landscape layers would be updated as a result), instead of providing ad-hoc functions on 
 	//  ALandscape, like we do currently (e.g. SetLayerAlpha, SetLayerVisibility, SetLayerName, etc.) on both the runtime code (ALandscape) and the UI code (FEdModeLandscape)
 	virtual void OnLayerCreated(FLandscapeLayer& Layer) {}
+
+	// Called by landscape after removing this layer from its list so that the layer can do
+	// any cleanup that it might need to do.
+	// TODO: Should this be protected and then we friend ALandscape?
+	virtual void OnLayerRemoved() {}
+
+	// UObject
+	LANDSCAPE_API virtual void PostLoad() override;
+protected:
+
+	// TODO: This might be removed once more things are moved from FLandscapeLayer to ULandscapeLayer
+	UPROPERTY()
+	TWeakObjectPtr<ALandscape> OwningLandscape;
 };
 
 /** 
diff --git a/Engine/Source/Runtime/Landscape/Private/GlobalMergeLegacySupportUtil.cpp b/Engine/Source/Runtime/Landscape/Private/GlobalMergeLegacySupportUtil.cpp
new file mode 100644
--- /dev/null
+++ b/Engine/Source/Runtime/Landscape/Private/GlobalMergeLegacySupportUtil.cpp
@@ -0,0 +1,45 @@
+// Copyright Epic Games, Inc. All Rights Reserved.
+
+#include "GlobalMergeLegacySupportUtil.h"
+
+#include "LandscapeBlueprintBrushBase.h" // FLandscapeBrushParameters
+#include "LandscapeEditTypes.h" // ELandscapeToolTargetType
+#include "Engine/TextureRenderTarget2D.h"
+
+#if WITH_EDITOR
+UTextureRenderTarget2D* ILandscapeBrushRenderCallAdapter_GlobalMergeLegacySupport::RenderAsBlueprintBrush(
+	const FLandscapeBrushParameters& InParameters, const FTransform& NewLandscapeTransform, const FIntPoint& RenderAreaSize)
+{
+	TRACE_CPUPROFILER_EVENT_SCOPE(ILandscapeBrushRenderCallAdapter_GlobalMergeLegacySupport::RenderAsBlueprintBrush);
+
+	// Do the same early outs and processing as rendering a blueprint brush
+
+	if (((InParameters.LayerType == ELandscapeToolTargetType::Heightmap) && !AffectsHeightmapAsBlueprintBrush())
+		|| ((InParameters.LayerType == ELandscapeToolTargetType::Weightmap) && !AffectsWeightmapLayerAsBlueprintBrush(InParameters.WeightmapLayerName))
+		|| ((InParameters.LayerType == ELandscapeToolTargetType::Visibility) && !AffectsVisibilityLayerAsBlueprintBrush()))
+	{
+		return nullptr;
+	}
+
+	if (InParameters.CombinedResult == nullptr)
+	{
+		return nullptr;
+	}
+
+	const FIntPoint NewLandscapeRenderTargetSize = FIntPoint(InParameters.CombinedResult->SizeX, InParameters.CombinedResult->SizeY);
+	if (!CurrentRenderAreaWorldTransform.Equals(NewLandscapeTransform) 
+		|| (CurrentRenderAreaSize != RenderAreaSize)
+		|| CurrentRenderTargetSize != NewLandscapeRenderTargetSize)
+	{
+		CurrentRenderAreaWorldTransform = NewLandscapeTransform;
+		CurrentRenderAreaSize = RenderAreaSize;
+		CurrentRenderTargetSize = NewLandscapeRenderTargetSize;
+
+		InitializeAsBlueprintBrush(CurrentRenderAreaWorldTransform, CurrentRenderAreaSize, CurrentRenderTargetSize);
+	}
+
+	UTextureRenderTarget2D* Result = RenderLayerAsBlueprintBrush(InParameters);
+
+	return Result;
+}
+#endif
\ No newline at end of file
diff --git a/Engine/Source/Runtime/Landscape/Private/Landscape.cpp b/Engine/Source/Runtime/Landscape/Private/Landscape.cpp
--- a/Engine/Source/Runtime/Landscape/Private/Landscape.cpp
+++ b/Engine/Source/Runtime/Landscape/Private/Landscape.cpp
@@ -2875,8 +2875,12 @@ PRAGMA_DISABLE_DEPRECATION_WARNINGS
 			}
 			check(Layer.EditLayer == nullptr);
 			Layer.EditLayer = NewObject<ULandscapeEditLayerBase>(this, EditLayerClass, MakeUniqueObjectName(this, EditLayerClass));
+			Layer.EditLayer->SetBackPointer(this);
 			Layer.EditLayer->OnLayerCreated(Layer);
 		}
+
+		// Empty the old property now that we've moved them over, else we'll accidentally keep references to brushes etc.
+		LandscapeLayers_DEPRECATED.Empty();
 PRAGMA_ENABLE_DEPRECATION_WARNINGS
 	}
 
diff --git a/Engine/Source/Runtime/Landscape/Private/LandscapeEditLayer.cpp b/Engine/Source/Runtime/Landscape/Private/LandscapeEditLayer.cpp
--- a/Engine/Source/Runtime/Landscape/Private/LandscapeEditLayer.cpp
+++ b/Engine/Source/Runtime/Landscape/Private/LandscapeEditLayer.cpp
@@ -6,6 +6,19 @@
 
 #define LOCTEXT_NAMESPACE "LandscapeEditLayer"
 
+void ULandscapeEditLayerBase::SetBackPointer(ALandscape* Landscape)
+{
+	OwningLandscape = Landscape;
+}
+
+void ULandscapeEditLayerBase::PostLoad()
+{
+	Super::PostLoad();
+
+	// TODO[jonathan.bard] Remove
+	// Needed because we might have saved some layers before we realized we were missing this flag
+	SetFlags(RF_Transactional);
+}
 
 // ----------------------------------------------------------------------------------
 
@@ -57,4 +70,4 @@ TArray<ULandscapeEditLayerSplines::FEditLayerAction> ULandscapeEditLayerSplines:
 	return Actions;
 }
 
-#undef LOCTEXT_NAMESPACE
\ No newline at end of file
+#undef LOCTEXT_NAMESPACE
diff --git a/Engine/Source/Runtime/Landscape/Private/LandscapeEditLayers.cpp b/Engine/Source/Runtime/Landscape/Private/LandscapeEditLayers.cpp
--- a/Engine/Source/Runtime/Landscape/Private/LandscapeEditLayers.cpp
+++ b/Engine/Source/Runtime/Landscape/Private/LandscapeEditLayers.cpp
@@ -55,6 +55,7 @@ LandscapeEditLayers.cpp: Landscape editing layers mode
 #include "LandscapeRender.h"
 #include "LandscapeInfoMap.h"
 #include "Misc/MessageDialog.h"
+#include "GlobalMergeLegacySupportUtil.h" // ILandscapeBrushRenderCallAdapter_GlobalMergeLegacySupport
 #include "GameFramework/WorldSettings.h"
 #include "UObject/UObjectThreadContext.h"
 #include "LandscapeSplinesComponent.h"
@@ -4967,6 +4968,29 @@ int32 ALandscape::PerformLayersHeightmapsGlobalMerge(const FUpdateLayersContentC
 					PrintLayersDebugRT(FString::Printf(TEXT("LS Height: %s Component %s += -> CombinedNonAtlas %s"), *Layer.Name.ToString(), *BrushOutputNonAtlasRT->GetName(), *CombinedHeightmapNonAtlasRT->GetName()), CombinedHeightmapNonAtlasRT);
 				}
 			}
+
+			// Legacy global merge support for non-blueprint-brush renderers: allow edit layer to act as blueprint brush.
+			ILandscapeBrushRenderCallAdapter_GlobalMergeLegacySupport* RenderCallAdapter = Cast<ILandscapeBrushRenderCallAdapter_GlobalMergeLegacySupport>(Layer.EditLayer);
+			if (RenderCallAdapter)
+			{
+				FLandscapeBrushParameters BrushParameters = BuildLandscapeBrushParameters(
+					/*bIsHeightmapMerge = */true, LandscapeExtent, CombinedHeightmapNonAtlasRT);
+
+				UTextureRenderTarget2D* BrushOutputNonAtlasRT = RenderCallAdapter->RenderAsBlueprintBrush(BrushParameters, GetTransform(), 
+					// TODO: Someday this will be part of BrushParameters
+					LandscapeExtent.Max - LandscapeExtent.Min);
+
+				// Do the same conditional copy that we do for blueprint brushes
+				if (BrushOutputNonAtlasRT != CombinedHeightmapNonAtlasRT
+					&& BrushOutputNonAtlasRT
+					&& BrushOutputNonAtlasRT->SizeX == CombinedHeightmapNonAtlasRT->SizeX
+					&& BrushOutputNonAtlasRT->SizeY == CombinedHeightmapNonAtlasRT->SizeY)
+				{
+					ExecuteCopyLayersTexture({ FLandscapeLayersCopyTextureParams(BrushOutputNonAtlasRT, CombinedHeightmapNonAtlasRT) });
+					PrintLayersDebugRT(FString::Printf(TEXT("LS Height: %s Component %s += -> CombinedNonAtlas %s"), *Layer.Name.ToString(), 
+						*BrushOutputNonAtlasRT->GetName(), *CombinedHeightmapNonAtlasRT->GetName()), CombinedHeightmapNonAtlasRT);
+				}
+			}
 		}
 
 		// copy CombinedHeightmapNonAtlasRT to LandscapeScratchRT3 (as a source for later layers... this is wasted on the last layer I think...)
@@ -5079,6 +5103,24 @@ int32 ALandscape::PerformLayersHeightmapsGlobalMerge(const FUpdateLayersContentC
 	return InMergeParams.HeightmapUpdateModes;
 }
 
+FLandscapeBrushParameters ALandscape::BuildLandscapeBrushParameters(bool bInIsHeightmapMerge, 
+	const FIntRect& InRenderAreaExtents, UTextureRenderTarget2D* InRenderTarget, FName InWeightmapLayerName)
+{
+	FTransform RenderAreaWorldTransform = GetTransform();
+	FVector OffsetVector(InRenderAreaExtents.Min.X, InRenderAreaExtents.Min.Y, 0.f);
+	FVector Translation = RenderAreaWorldTransform.TransformFVector4(OffsetVector);
+	RenderAreaWorldTransform.SetTranslation(Translation);
+	FIntPoint RenderAreaSize = InRenderAreaExtents.Max - InRenderAreaExtents.Min;
+
+	// TODO: Someday we will build the final return struct like this:
+	// return FLandscapeBrushParameters(bInIsHeightmapMerge, RenderAreaWorldTransform, RenderAreaSize, InRenderTarget, InWeightmapLayerName);
+
+	return FLandscapeBrushParameters(bInIsHeightmapMerge ? ELandscapeToolTargetType::Heightmap 
+		: (InWeightmapLayerName == UMaterialExpressionLandscapeVisibilityMask::ParameterName) ? ELandscapeToolTargetType::Visibility 
+		: ELandscapeToolTargetType::Weightmap,
+		InRenderTarget, InWeightmapLayerName);
+}
+
 int32 ALandscape::RegenerateLayersHeightmaps(const FUpdateLayersContentContext& InUpdateLayersContentContext)
 {
 	TRACE_CPUPROFILER_EVENT_SCOPE(LandscapeLayers_RegenerateLayersHeightmaps);
@@ -7162,6 +7204,27 @@ int32 ALandscape::PerformLayersWeightmapsGlobalMerge(FUpdateLayersContentContext
 								bHasWeightmapData = true;
 							}
 						}
+
+						// Legacy global merge support for non-blueprint-brush renderers: allow edit layer to act as blueprint brush.
+						ILandscapeBrushRenderCallAdapter_GlobalMergeLegacySupport* RenderCallAdapter = Cast<ILandscapeBrushRenderCallAdapter_GlobalMergeLegacySupport>(Layer.EditLayer);
+						if (RenderCallAdapter && !LayerInfoObjects.Contains(InfoLayerSettings.LayerInfoObj))
+						{
+							TOptional<int32> LayerInfoSettingsAllocatedIndex;
+							if (RenderCallAdapter->AffectsWeightmapLayerAsBlueprintBrush(InfoLayerSettings.GetLayerName()))
+							{
+								LayerInfoSettingsAllocatedIndex = LayerInfoSettingsIndex + 1; // due to visibility layer that is at 0
+							}
+							else if (RenderCallAdapter->AffectsVisibilityLayerAsBlueprintBrush() && UE::Landscape::IsVisibilityLayer(InfoLayerSettings.LayerInfoObj))
+							{
+								LayerInfoSettingsAllocatedIndex = GetVisibilityLayerAllocationIndex();
+							}
+
+							if (LayerInfoSettingsAllocatedIndex.IsSet())
+							{
+								LayerInfoObjects.Add(InfoLayerSettings.LayerInfoObj, LayerInfoSettingsAllocatedIndex.GetValue());
+								bHasWeightmapData = true;
+							}
+						}
 					}
 				}
 			}
@@ -7342,6 +7405,33 @@ int32 ALandscape::PerformLayersWeightmapsGlobalMerge(FUpdateLayersContentContext
 							}
 						}
 
+						// Legacy global merge support for non-blueprint-brush renderers: allow edit layer to act as blueprint brush.
+						ILandscapeBrushRenderCallAdapter_GlobalMergeLegacySupport* RenderCallAdapter = Cast<ILandscapeBrushRenderCallAdapter_GlobalMergeLegacySupport>(Layer.EditLayer);
+						if (RenderCallAdapter)
+						{
+							FLandscapeBrushParameters BrushParameters = BuildLandscapeBrushParameters(
+								/*bIsHeightmapMerge = */false, LandscapeExtent, LandscapeScratchRT3, LayerInfoObj->LayerName);
+							
+							UTextureRenderTarget2D* BrushOutputRT = RenderCallAdapter->RenderAsBlueprintBrush(BrushParameters, GetTransform(),
+								// TODO: Someday this will be part of BrushParameters
+								LandscapeExtent.Max - LandscapeExtent.Min);
+
+							if (BrushOutputRT != nullptr && BrushOutputRT->SizeX == LandscapeScratchRT3->SizeX 
+								&& BrushOutputRT->SizeY == LandscapeScratchRT3->SizeY)
+							{
+								BrushRequiredAllocations.AddUnique(LayerInfoObj);
+
+								// Same conditional copy as for blueprint brushes
+								if (BrushOutputRT != LandscapeScratchRT3)
+								{
+									SourceDebugName = FString::Printf(TEXT("Weight: %s PaintLayer: %s Brush: %s"), *Layer.Name.ToString(), *LayerInfoObj->LayerName.ToString(), *BrushOutputRT->GetName());
+									DestDebugName = LandscapeScratchRT3->GetName();
+									ExecuteCopyLayersTexture({ FLandscapeLayersCopyTextureParams(SourceDebugName, BrushOutputRT->GameThread_GetRenderTargetResource(), DestDebugName, LandscapeScratchRT3->GameThread_GetRenderTargetResource()) });
+									PrintLayersDebugRT(FString::Printf(TEXT("LS Weight: %s Component %s += -> Combined %s"), *Layer.Name.ToString(), *BrushOutputRT->GetName(), *LandscapeScratchRT3->GetName()), LandscapeScratchRT3);
+								}
+							}
+						}
+
 						PrintLayersDebugRT(FString::Printf(TEXT("LS Weight: %s CombinedPostBrushProcLayerWeightmap -> Paint Layer RT %s"), *Layer.Name.ToString(), *LandscapeScratchRT3->GetName()), LandscapeScratchRT3, 0, false);
 
 						SourceDebugName = FString::Printf(TEXT("Weight: %s PaintLayer: %s %s"), *Layer.Name.ToString(), *LayerInfoObj->LayerName.ToString(), *LandscapeScratchRT3->GetName());
@@ -9327,14 +9417,14 @@ void ALandscape::DeleteLayer(int32 InLayerIndex)
 	ensure(HasLayersContent());
 
 	ULandscapeInfo* LandscapeInfo = GetLandscapeInfo();
-	const FLandscapeLayer* Layer = GetLayerConst(InLayerIndex);
-	if (!LandscapeInfo || !Layer)
+	const FLandscapeLayer* LayerStruct = GetLayerConst(InLayerIndex);
+	if (!LandscapeInfo || !LayerStruct)
 	{
 		return;
 	}
 
 	Modify();
-	FGuid LayerGuid = Layer->Guid;
+	FGuid LayerGuid = LayerStruct->Guid;
 
 	// Clean up Weightmap usage in LandscapeProxies
 	LandscapeInfo->ForEachLandscapeProxy([&LayerGuid](ALandscapeProxy* Proxy)
@@ -9343,8 +9433,18 @@ void ALandscape::DeleteLayer(int32 InLayerIndex)
 		return true;
 	});
 
+	// We're about to remove the layer from our list, which will invalidate our LayerStruct pointer.
+	// We'll need to call OnLayerRemoved afterward, though, so keep pointer to the UObject.
+	ULandscapeEditLayerBase* Layer = LayerStruct->EditLayer;
+
 	// Remove layer from list
 	LandscapeEditLayers.RemoveAt(InLayerIndex);
+	LayerStruct = nullptr;
+
+	if (Layer)
+	{
+		Layer->OnLayerRemoved();
+	}
 
 	// Request Update
 	RequestLayersContentUpdateForceAll();
@@ -9918,6 +10018,12 @@ void ALandscape::OnLayerCreatedInternal(FLandscapeLayer& Layer)
 	ULandscapeInfo* LandscapeInfo = GetLandscapeInfo();
 	check(LandscapeInfo != nullptr); 
 
+	// TODO: Might not be necessary eventually, if EditLayer has ability to trigger landscape updates
+	// and has access to its guid.
+	Layer.EditLayer->SetBackPointer(this);
+	
+	Layer.EditLayer->SetFlags(RF_Transactional);
+
 	Layer.EditLayer->OnLayerCreated(Layer);
 
 	// Create associated layer data in each landscape proxy
diff --git a/Engine/Source/Runtime/Landscape/Public/GlobalMergeLegacySupportUtil.h b/Engine/Source/Runtime/Landscape/Public/GlobalMergeLegacySupportUtil.h
new file mode 100644
--- /dev/null
+++ b/Engine/Source/Runtime/Landscape/Public/GlobalMergeLegacySupportUtil.h
@@ -0,0 +1,66 @@
+// Copyright Epic Games, Inc. All Rights Reserved.
+
+#pragma once
+
+#include "Math/MathFwd.h" // FTransform, FIntPoint
+#include "Misc/CoreMiscDefines.h" // PURE_VIRTUAL
+#include "UObject/Interface.h"
+
+#include "GlobalMergeLegacySupportUtil.generated.h"
+
+// Note: this file is likely to be removed once global merge is fully phased out.
+
+class UTextureRenderTarget2D;
+struct FLandscapeBrushParameters;
+
+// UObject for the UInterface
+UINTERFACE(MinimalAPI)
+class ULandscapeBrushRenderCallAdapter_GlobalMergeLegacySupport : public UInterface
+{
+	GENERATED_BODY()
+};
+
+/**
+ * Interface that allows an edit layer UObject to receive a render call the way that a blueprint
+ * brush would in global merge mode, so that an edit layer UObject that implements custom batched
+ * merge functions can still perform its work if global merge mode is used instead.
+ * 
+ * This interface is likely to be removed once global merge is no longer used.
+ */
+class ILandscapeBrushRenderCallAdapter_GlobalMergeLegacySupport
+{
+	GENERATED_BODY()
+
+#if WITH_EDITOR
+public:
+	/**
+	 * In global merge mode, this function is called in the same places that FLandscapeLayerBrush::RenderLayer(BrushParameters)
+	 * is called. This function calls the virtual InitializeAsBlueprintBrush and RenderLayerAsBlueprintBrush
+	 * functions, which are equivalent to a blueprint brush actor's Initialize_Native and RenderLayer_Native.
+	 */
+	UTextureRenderTarget2D* RenderAsBlueprintBrush(const FLandscapeBrushParameters& InParameters, const FTransform& LandscapeTransform, 
+		// TODO: Someday will be part of FLandscapeBrushParameters
+		const FIntPoint& RenderAreaSize);
+
+	// Called in the same places as the equivalent methods on FLandscapeLayerBrush.
+	virtual bool AffectsHeightmapAsBlueprintBrush() const { return false; }
+	virtual bool AffectsWeightmapLayerAsBlueprintBrush(const FName& InLayerName) const { return false; }
+	virtual bool AffectsVisibilityLayerAsBlueprintBrush() const { return false; }
+
+protected:
+	
+	// Overridable, equivalent to Initialize_Native on a blueprint brush actor
+	virtual void InitializeAsBlueprintBrush(const FTransform& InLandscapeTransform,
+		const FIntPoint& InLandscapeSize,
+		const FIntPoint& InLandscapeRenderTargetSize) {}
+
+	// Overridable, equivalent to RenderLayer_Native on a blueprint brush actor.
+	virtual UTextureRenderTarget2D* RenderLayerAsBlueprintBrush(const FLandscapeBrushParameters& InParameters)
+		PURE_VIRTUAL(ILandscapeBrushRenderCallAdapter_GlobalMergeLegacySupport::InitializeAsBlueprintBrush, return nullptr; );
+
+	// Used in RenderAsBlueprintBrush to cache things the same way that blueprint brushes do
+	FTransform CurrentRenderAreaWorldTransform;
+	FIntPoint CurrentRenderAreaSize;
+	FIntPoint CurrentRenderTargetSize;
+#endif
+};
