LandscapePatch: Fix static analyzer warning.

diff --git a/Engine/Plugins/Experimental/LandscapePatch/Source/LandscapePatch/Private/LandscapePatchComponent.cpp b/Engine/Plugins/Experimental/LandscapePatch/Source/LandscapePatch/Private/LandscapePatchComponent.cpp
--- a/Engine/Plugins/Experimental/LandscapePatch/Source/LandscapePatch/Private/LandscapePatchComponent.cpp
+++ b/Engine/Plugins/Experimental/LandscapePatch/Source/LandscapePatch/Private/LandscapePatchComponent.cpp
@@ -840,7 +840,8 @@ bool ULandscapePatchComponent::BindToEditLayer(const FGuid& Guid)
 	EditLayer = Cast<ULandscapePatchEditLayer>(LayerStruct ? LayerStruct->EditLayer : nullptr);
 	EditLayerGuid = Guid;
 	DetailPanelLayerGuid = Guid.ToString();
-	DetailPanelLayerName = EditLayer.IsValid() ? LayerStruct->Name.ToString()
+	DetailPanelLayerName = 
+		EditLayer.IsValid() && LayerStruct ? LayerStruct->Name.ToString()
 		: LayerStruct ? TEXT("-Incorrect Layer Type-")
 		: TEXT("-Layer Not Found-");
 
