LandscapePatch: Fix issues with internal render target serialization introduced in 33341608. Fix incorrect priority handling for newly placed patches.

diff --git a/Engine/Plugins/Experimental/LandscapePatch/Source/LandscapePatch/Private/LandscapePatchComponent.cpp b/Engine/Plugins/Experimental/LandscapePatch/Source/LandscapePatch/Private/LandscapePatchComponent.cpp
--- a/Engine/Plugins/Experimental/LandscapePatch/Source/LandscapePatch/Private/LandscapePatchComponent.cpp
+++ b/Engine/Plugins/Experimental/LandscapePatch/Source/LandscapePatch/Private/LandscapePatchComponent.cpp
@@ -323,6 +323,15 @@ void ULandscapePatchComponent::OnComponentCreated()
 		return;
 	}
 
+	// We're going to be binding to an edit layer, which will place us in its registered patch list with
+	//  our current priority. If we're later going to be updating our priority to be the highest, then 
+	//  we need to temporarily lower our priority so that we aren't accidentally the highest priority patch
+	//  the edit layer sees when we query it. 
+	if (IsPatchInWorld() && PriorityInitialization == ELandscapePatchPriorityInitialization::AcquireHighest)
+	{
+		Priority = TNumericLimits<double>::Lowest();
+	}
+
 	// Otherwise, bind to some edit layer
 	bool bConnectionToLandscapeEstablished = false;
 	if (EditLayerGuid.IsValid())
@@ -352,18 +361,16 @@ void ULandscapePatchComponent::OnComponentCreated()
 			break;
 		case ELandscapePatchPriorityInitialization::SmallIncrement:
 			Priority += 0.01;
+			if (EditLayer.IsValid())
+			{
+				EditLayer->NotifyOfPriorityChange(this);
+			}
 			break;
 		case ELandscapePatchPriorityInitialization::AcquireHighest:
 			if (EditLayer.IsValid())
 			{
-				if (IsPatchPreview())
-				{
-					Priority = EditLayer->GetHighestPatchPriority();
-				}
-				else
-				{
-					Priority = EditLayer->ClaimHighestPatchPriority();
-				}
+				Priority = EditLayer->GetHighestPatchPriority() + 1;
+				EditLayer->NotifyOfPriorityChange(this);
 			}
 			break;
 		}
diff --git a/Engine/Plugins/Experimental/LandscapePatch/Source/LandscapePatch/Private/LandscapePatchEditLayer.cpp b/Engine/Plugins/Experimental/LandscapePatch/Source/LandscapePatch/Private/LandscapePatchEditLayer.cpp
--- a/Engine/Plugins/Experimental/LandscapePatch/Source/LandscapePatch/Private/LandscapePatchEditLayer.cpp
+++ b/Engine/Plugins/Experimental/LandscapePatch/Source/LandscapePatch/Private/LandscapePatchEditLayer.cpp
@@ -244,12 +244,6 @@ double ULandscapePatchEditLayer::GetHighestPatchPriority()
 {
 	return HighestKnownPriority;
 }
-
-double ULandscapePatchEditLayer::ClaimHighestPatchPriority()
-{
-	HighestKnownPriority += 1;
-	return HighestKnownPriority;
-}
 #endif // WITH_EDITOR
 
 void ULandscapePatchEditLayer::GetRenderDependencies(TSet<UObject*>& OutDependencies) const
diff --git a/Engine/Plugins/Experimental/LandscapePatch/Source/LandscapePatch/Private/LandscapeTexturePatch.cpp b/Engine/Plugins/Experimental/LandscapePatch/Source/LandscapePatch/Private/LandscapeTexturePatch.cpp
--- a/Engine/Plugins/Experimental/LandscapePatch/Source/LandscapePatch/Private/LandscapeTexturePatch.cpp
+++ b/Engine/Plugins/Experimental/LandscapePatch/Source/LandscapePatch/Private/LandscapeTexturePatch.cpp
@@ -252,10 +252,10 @@ UTextureRenderTarget2D* ULandscapeTexturePatch::ApplyToWeightmap(ULandscapeWeigh
 	case ELandscapeTexturePatchSourceMode::None:
 		return InCombinedResult;
 	case ELandscapeTexturePatchSourceMode::InternalTexture:
-		PatchUObject = PatchInfo->InternalData ? PatchInfo->InternalData->GetInternalTexture() : nullptr;
+		PatchUObject = GetWeightPatchInternalTexture(PatchInfo);
 		break;
 	case ELandscapeTexturePatchSourceMode::TextureBackedRenderTarget:
-		PatchUObject = PatchInfo->InternalData ? PatchInfo->InternalData->GetRenderTarget() : nullptr;
+		PatchUObject = GetWeightPatchRenderTarget(PatchInfo);
 		break;
 	case ELandscapeTexturePatchSourceMode::TextureAsset:
 		if (IsValid(PatchInfo->TextureAsset) && !ensureMsgf(PatchInfo->TextureAsset->VirtualTextureStreaming == 0,
@@ -905,36 +905,6 @@ void ULandscapeTexturePatch::GetRenderDependencies(TSet<UObject*>& OutDependenci
 	}
 }
 
-void ULandscapeTexturePatch::MakeSureInternalDataIsAllocated()
-{
-	if (!ensure(!IsTemplate()))
-	{
-		return;
-	}
-
-	// If we're instancing a blueprint actor, then it might be missing the needed internal data if the source mode
-	// was set to internal texture or render target, because we can't reliably create those in the archetype.
-	// So, we initialize that data if it's missing but needed when instancing.
-	if (!HeightInternalData && (HeightSourceMode == ELandscapeTexturePatchSourceMode::InternalTexture
-		|| HeightSourceMode == ELandscapeTexturePatchSourceMode::TextureBackedRenderTarget))
-	{
-		TransitionHeightSourceModeInternal(ELandscapeTexturePatchSourceMode::None, HeightSourceMode);
-	}
-
-	for (const TObjectPtr<ULandscapeWeightPatchTextureInfo>& WeightPatch : WeightPatches)
-	{
-		if (!IsValid(WeightPatch))
-		{
-			continue;
-		}
-		if (!WeightPatch->InternalData && (WeightPatch->SourceMode == ELandscapeTexturePatchSourceMode::InternalTexture
-			|| WeightPatch->SourceMode == ELandscapeTexturePatchSourceMode::TextureBackedRenderTarget))
-		{
-			WeightPatch->TransitionSourceModeInternal(ELandscapeTexturePatchSourceMode::None, WeightPatch->SourceMode);
-		}
-	}
-}
-
 #endif
 
 void ULandscapeTexturePatch::SnapToLandscape()
@@ -1311,17 +1281,70 @@ void ULandscapeTexturePatch::SetHeightTextureAsset(UTexture* TextureIn)
 
 UTextureRenderTarget2D* ULandscapeTexturePatch::GetHeightRenderTarget(bool bMarkDirty)
 {
+#if WITH_EDITOR
+
+	if (IsTemplate())
+	{
+		return nullptr;
+	}
+
 	if (bMarkDirty)
 	{
 		MarkPackageDirty();
 	}
-	return HeightInternalData ? HeightInternalData->GetRenderTarget() : nullptr;
+
+	// In templates (i.e. in blueprint editor), it's not safe to create textures, so if we are an instantiation
+	//  of a blueprint, we may not yet have the internal render target allocated. It might seem like a good idea
+	//  to do this in OnComponentCreated, but that causes default construction script instance data application
+	//  to see the data as modified, and prevents it from being carried over properly (see usage of GetUCSModifiedProperties
+	//  in ComponentInstanceDataCache.cpp). Doing it in ApplyComponentInstanceData also seems to be a good idea at
+	//  first, but we can't do it in ECacheApplyPhase::PostSimpleConstructionScript for the same reason as OnComponentModified,
+	//  and doing it in ECacheApplyPhase::PostUserConstructionScript is too late because the user may want to write
+	//  to the render target in the user construction script.
+	// So, we do this allocation right when the render target is requested.
+	if (HeightSourceMode == ELandscapeTexturePatchSourceMode::TextureBackedRenderTarget)
+	{
+		if (!HeightInternalData || !HeightInternalData->GetRenderTarget())
+		{
+			TransitionHeightSourceModeInternal(ELandscapeTexturePatchSourceMode::None, HeightSourceMode);
+		}
+
+		return ensure(HeightInternalData) ? HeightInternalData->GetRenderTarget() : nullptr;
+	}
+#endif
+
+	return nullptr;
+}
+
+UTexture2D* ULandscapeTexturePatch::GetHeightInternalTexture()
+{
+#if WITH_EDITOR
+
+	if (IsTemplate())
+	{
+		return nullptr;
+	}
+
+	// Allocate data if needed (see comment in GetHeightRenderTarget)
+	if (HeightSourceMode == ELandscapeTexturePatchSourceMode::TextureBackedRenderTarget
+		|| HeightSourceMode == ELandscapeTexturePatchSourceMode::InternalTexture)
+	{
+		if (!HeightInternalData || !HeightInternalData->GetInternalTexture())
+		{
+			TransitionHeightSourceModeInternal(ELandscapeTexturePatchSourceMode::None, HeightSourceMode);
+		}
+
+		return ensure(HeightInternalData) ? HeightInternalData->GetInternalTexture() : nullptr;
+	}
+#endif
+
+	return nullptr;
 }
 
 void ULandscapeTexturePatch::UpdateHeightConvertToNativeParamsIfNeeded()
 {
 #if WITH_EDITOR
-	if (HeightInternalData && Landscape.IsValid() && PatchManager.IsValid())
+	if (HeightInternalData)
 	{
 		FLandscapeHeightPatchConvertToNativeParams ConversionParams = GetHeightConvertToNativeParams();
 		if (ConversionParams.HeightScale == 0)
@@ -1537,27 +1560,72 @@ ELandscapeTexturePatchSourceMode ULandscapeTexturePatch::GetWeightPatchSourceMod
 	return ELandscapeTexturePatchSourceMode::None;
 }
 
-UTextureRenderTarget2D* ULandscapeTexturePatch::GetWeightPatchRenderTarget(const FName& InWeightmapLayerName)
+UTextureRenderTarget2D* ULandscapeTexturePatch::GetWeightPatchRenderTarget(const FName& InWeightmapLayerName, bool bMarkDirty)
 {
+	if (IsTemplate())
+	{
+		return nullptr;
+	}
+
 	for (const TObjectPtr<ULandscapeWeightPatchTextureInfo>& WeightPatch : WeightPatches)
 	{
 		if (IsValid(WeightPatch) && WeightPatch->WeightmapLayerName == InWeightmapLayerName)
 		{
-			return WeightPatch->InternalData ? WeightPatch->InternalData->GetRenderTarget() : nullptr;
+			if (bMarkDirty)
+			{
+				MarkPackageDirty();
+			}
+
+			return GetWeightPatchRenderTarget(WeightPatch);
 		}
 	}
 	return nullptr;
 }
 
-UTexture2D* ULandscapeTexturePatch::GetWeightPatchInternalTexture(const FName& InWeightmapLayerName)
+UTextureRenderTarget2D* ULandscapeTexturePatch::GetWeightPatchRenderTarget(ULandscapeWeightPatchTextureInfo* WeightPatch)
 {
-	for (const TObjectPtr<ULandscapeWeightPatchTextureInfo>& WeightPatch : WeightPatches)
+#if WITH_EDITOR
+	if (IsTemplate() || !IsValid(WeightPatch))
 	{
-		if (IsValid(WeightPatch) && WeightPatch->WeightmapLayerName == InWeightmapLayerName)
+		return nullptr;
+	}
+
+	// Allocate data if needed (see comment in GetHeightRenderTarget)
+	if (WeightPatch->SourceMode == ELandscapeTexturePatchSourceMode::TextureBackedRenderTarget)
+	{
+		if (!WeightPatch->InternalData || !WeightPatch->InternalData->GetRenderTarget())
+		{
+			WeightPatch->TransitionSourceModeInternal(ELandscapeTexturePatchSourceMode::None, WeightPatch->SourceMode);
+		}
+
+		return ensure(WeightPatch->InternalData) ? WeightPatch->InternalData->GetRenderTarget() : nullptr;
+	}
+#endif
+
+	return nullptr;
+}
+
+UTexture2D* ULandscapeTexturePatch::GetWeightPatchInternalTexture(ULandscapeWeightPatchTextureInfo* WeightPatch)
+{
+#if WITH_EDITOR
+	if (IsTemplate() || !IsValid(WeightPatch))
+	{
+		return nullptr;
+	}
+
+	// Allocate data if needed (see comment in GetHeightRenderTarget)
+	if (WeightPatch->SourceMode == ELandscapeTexturePatchSourceMode::InternalTexture 
+		|| WeightPatch->SourceMode == ELandscapeTexturePatchSourceMode::TextureBackedRenderTarget)
+	{
+		if (!WeightPatch->InternalData || !WeightPatch->InternalData->GetInternalTexture())
 		{
-			return WeightPatch->InternalData ? WeightPatch->InternalData->GetInternalTexture() : nullptr;
+			WeightPatch->TransitionSourceModeInternal(ELandscapeTexturePatchSourceMode::None, WeightPatch->SourceMode);
 		}
+
+		return ensure(WeightPatch->InternalData) ? WeightPatch->InternalData->GetInternalTexture() : nullptr;
 	}
+#endif
+
 	return nullptr;
 }
 
@@ -1617,117 +1685,3 @@ void ULandscapeTexturePatch::SetEditVisibilityLayer(const FName& InWeightmapLaye
 		}
 	}
 }
-
-#if WITH_EDITOR
-TStructOnScope<FActorComponentInstanceData> ULandscapeTexturePatch::GetComponentInstanceData() const
-{
-	return MakeStructOnScope<FActorComponentInstanceData, FLandscapeTexturePatchInstanceData>(this);
-}
-#endif
-
-FLandscapeTexturePatchInstanceData::FLandscapeTexturePatchInstanceData(const ULandscapeTexturePatch* Patch)
-	: FLandscapePatchComponentInstanceData(Patch)
-{
-#if WITH_EDITOR
-
-	if (!ensure(Patch))
-	{
-		return;
-	}
-
-	HeightInternalData = Patch->HeightInternalData;
-	if (HeightInternalData.IsValid())
-	{
-		// Disconnect owner so that the object doesn't get destroyed
-		HeightInternalData->Rename(nullptr, GetTransientPackage());
-	}
-
-	for (const TObjectPtr<ULandscapeWeightPatchTextureInfo>& WeightPatch : Patch->WeightPatches)
-	{
-		if (!IsValid(WeightPatch))
-		{
-			continue;
-		}
-
-		WeightPatchInternalData.Add(WeightPatch->InternalData);
-		
-		if (IsValid(WeightPatch->InternalData))
-		{
-			// Disconnect owner so that the object doesn't get destroyed
-			WeightPatch->InternalData->Rename(nullptr, GetTransientPackage());
-		}
-	}
-#endif
-}
-
-// Called before/after rerunning construction scripts
-void ULandscapeTexturePatch::ApplyComponentInstanceData(FLandscapeTexturePatchInstanceData* ComponentInstanceData, ECacheApplyPhase CacheApplyPhase)
-{
-#if WITH_EDITOR
-
-	if (CacheApplyPhase == ECacheApplyPhase::PostUserConstructionScript)
-	{
-		// Avoid stomping user construction script changes.
-		return;
-	}
-
-	UWorld* World = GetWorld();
-	if (!ComponentInstanceData 
-		// TODO: Replace with !IsPatchInWorld() once that CL is in
-		|| IsTemplate() || !IsValid(this) || !IsValid(World) || World->WorldType != EWorldType::Editor)
-	{
-		return;
-	}
-
-	if (ComponentInstanceData->HeightInternalData.IsValid() 
-		&& (HeightSourceMode == ELandscapeTexturePatchSourceMode::InternalTexture
-			|| HeightSourceMode == ELandscapeTexturePatchSourceMode::TextureBackedRenderTarget))
-	{
-		HeightInternalData = ComponentInstanceData->HeightInternalData.Get();
-
-		// Attach to the new object
-		HeightInternalData->Rename(nullptr, this);
-	}
-	
-	for (int32 i = 0; i < ComponentInstanceData->WeightPatchInternalData.Num(); ++i)
-	{
-		if (i >= WeightPatches.Num() || !IsValid(WeightPatches[i]) || !ComponentInstanceData->WeightPatchInternalData[i].IsValid())
-		{
-			continue;
-		}
-
-		if (WeightPatches[i]->SourceMode == ELandscapeTexturePatchSourceMode::InternalTexture
-			|| WeightPatches[i]->SourceMode == ELandscapeTexturePatchSourceMode::TextureBackedRenderTarget)
-		{
-			WeightPatches[i]->InternalData = ComponentInstanceData->WeightPatchInternalData[i].Get();
-
-			// Attach to the new object
-			WeightPatches[i]->InternalData->Rename(nullptr, WeightPatches[i]);
-		}
-	}
-
-	// Note: For the below, one might think that we should override PostApplyToComponent, but that is actually
-	//  unreliable because it is called from FActorComponentInstanceData::ApplyToComponent, and gets skipped
-	//  if that subclass didn't happen to serialize some properties itself. Also, if it does get called, it is
-	//  not at the end because the Super call gets done earlier (not just by us, but by intermediate subclasses).
-
-	// If we didn't carry over internal data, we need to make sure at this point that it exists, because
-	//  we cannot create it in the template
-	MakeSureInternalDataIsAllocated();
-
-	// Make sure the blueprint-set resolution and format is applied
-	if (HeightInternalData)
-	{
-		HeightInternalData->SetResolution(ResolutionX, ResolutionY);
-		HeightInternalData->SetFormat(HeightRenderTargetFormat);
-	}
-	for (TObjectPtr<ULandscapeWeightPatchTextureInfo>& WeightPatch : WeightPatches)
-	{
-		if (IsValid(WeightPatch) && WeightPatch->InternalData)
-		{
-			WeightPatch->InternalData->SetResolution(ResolutionX, ResolutionY);
-		}
-	}
-	
-#endif // WITH_EDITOR
-}
diff --git a/Engine/Plugins/Experimental/LandscapePatch/Source/LandscapePatch/Public/LandscapePatchEditLayer.h b/Engine/Plugins/Experimental/LandscapePatch/Source/LandscapePatch/Public/LandscapePatchEditLayer.h
--- a/Engine/Plugins/Experimental/LandscapePatch/Source/LandscapePatch/Public/LandscapePatchEditLayer.h
+++ b/Engine/Plugins/Experimental/LandscapePatch/Source/LandscapePatch/Public/LandscapePatchEditLayer.h
@@ -44,12 +44,6 @@ public:
 	 */
 	virtual double GetHighestPatchPriority();
 
-	/**
-	 * Increments highest known patch priority and returns that value. Meant to allow an 
-	 *  already-registered patch to quickly become the highest priority patch.
-	 */
-	virtual double ClaimHighestPatchPriority();
-
 	void RequestLandscapeUpdate(bool bInUserTriggered = false);
 #endif
 
diff --git a/Engine/Plugins/Experimental/LandscapePatch/Source/LandscapePatch/Public/LandscapeTexturePatch.h b/Engine/Plugins/Experimental/LandscapePatch/Source/LandscapePatch/Public/LandscapeTexturePatch.h
--- a/Engine/Plugins/Experimental/LandscapePatch/Source/LandscapePatch/Public/LandscapeTexturePatch.h
+++ b/Engine/Plugins/Experimental/LandscapePatch/Source/LandscapePatch/Public/LandscapeTexturePatch.h
@@ -209,7 +209,6 @@ protected:
 #endif
 
 	friend class ULandscapeTexturePatch;
-	friend struct FLandscapeTexturePatchInstanceData;
 };
 
 
@@ -230,9 +229,6 @@ public:
 	virtual bool CanAffectVisibilityLayer() const override;
 	virtual void GetRenderDependencies(TSet<UObject*>& OutDependencies) const override;
 
-	// UActorComponent
-	virtual TStructOnScope<FActorComponentInstanceData> GetComponentInstanceData() const override;
-
 	// UObject
 	virtual void PostEditChangeProperty(FPropertyChangedEvent& PropertyChangedEvent) override;
 #endif
@@ -437,8 +433,12 @@ public:
 	UFUNCTION(BlueprintCallable, Category = LandscapePatch)
 	void SetWeightPatchTextureAsset(const FName& InWeightmapLayerName, UTexture* TextureIn);
 
+	/**
+	 * @param bMarkDirty If true, marks the containing package as dirty, since the render target is presumably
+	 *  being written to. Can be set to false if the render target is not being written to.
+	 */
 	UFUNCTION(BlueprintCallable, Category = LandscapePatch)
-	virtual UTextureRenderTarget2D* GetWeightPatchRenderTarget(const FName& InWeightmapLayerName);
+	virtual UTextureRenderTarget2D* GetWeightPatchRenderTarget(const FName& InWeightmapLayerName, bool bMarkDirty = true);
 
 	UFUNCTION(BlueprintCallable, Category = LandscapePatch)
 	virtual void SetUseAlphaChannelForWeightPatch(const FName& InWeightmapLayerName, bool bUseAlphaChannel);
@@ -452,11 +452,6 @@ public:
 	UFUNCTION(BlueprintCallable, Category = LandscapePatch)
 	virtual void SetEditVisibilityLayer(const FName& InWeightmapLayerName, const bool bEditVisibilityLayer);
 
-	// These need to be public so that we can take the internal textures and write them to external ones,
-	// but unclear whether we want to expose them to blueprints, since this is a fairly internal thing.
-	UTexture2D* GetHeightInternalTexture() { return HeightInternalData ? HeightInternalData->GetInternalTexture() : nullptr; ; }
-	UTexture2D* GetWeightPatchInternalTexture(const FName& InWeightmapLayerName);
-
 protected:
 	//~ Don't expose these on the instance because a user might not realize that they would lose their existing internal
 	//~ data by dragging them, and the only way they can reinitialize data in the viewport is through the methods that
@@ -605,8 +600,9 @@ protected:
 	int32 InitTextureSizeY = 33;
 
 private:
-	void ApplyComponentInstanceData(struct FLandscapeTexturePatchInstanceData* ComponentInstanceData, ECacheApplyPhase CacheApplyPhase);
-	friend struct FLandscapeTexturePatchInstanceData;
+	UTexture2D* GetHeightInternalTexture();
+	UTextureRenderTarget2D* GetWeightPatchRenderTarget(ULandscapeWeightPatchTextureInfo* WeightPatch);
+	UTexture2D* GetWeightPatchInternalTexture(ULandscapeWeightPatchTextureInfo* WeightPatch);
 
 	void UpdateHeightConvertToNativeParamsIfNeeded();
 #if WITH_EDITOR
@@ -631,7 +627,6 @@ private:
 	void ReinitializeWeightPatch(ULandscapeWeightPatchTextureInfo* PatchInfo, UTextureRenderTarget2D* InCombinedResult, 
 		const FTransform& LandscapeHeightmapToWorld);
 
-	void MakeSureInternalDataIsAllocated();
 	void ResetHeightRenderTargetFormat();
 #endif // WITH_EDITOR
 
@@ -639,37 +634,6 @@ private:
 	TEnumAsByte<ETextureRenderTargetFormat> HeightRenderTargetFormat = ETextureRenderTargetFormat::RTF_R32f;
 };
 
-
-/** Carries over internal data subobjects to avoid constant reconstruction on construction script reruns. */
-USTRUCT()
-struct FLandscapeTexturePatchInstanceData : public FLandscapePatchComponentInstanceData
-{
-	GENERATED_BODY()
-
-	FLandscapeTexturePatchInstanceData() = default;
-	FLandscapeTexturePatchInstanceData(const ULandscapeTexturePatch* SourceComponent);
-
-	virtual ~FLandscapeTexturePatchInstanceData() = default;
-
-	virtual bool ContainsData() const override
-	{
-		return true;
-	}
-
-	virtual void ApplyToComponent(UActorComponent* Component, const ECacheApplyPhase CacheApplyPhase) override
-	{
-		Super::ApplyToComponent(Component, CacheApplyPhase);
-		CastChecked<ULandscapeTexturePatch>(Component)->ApplyComponentInstanceData(this, CacheApplyPhase);
-	}
-
-#if WITH_EDITORONLY_DATA
-	UPROPERTY()
-	TWeakObjectPtr<ULandscapeHeightTextureBackedRenderTarget> HeightInternalData = nullptr;
-	UPROPERTY()
-	TArray<TWeakObjectPtr<ULandscapeWeightTextureBackedRenderTarget>> WeightPatchInternalData;
-#endif
-};
-
 #if UE_ENABLE_INCLUDE_ORDER_DEPRECATED_IN_5_2
 #include "CoreMinimal.h"
 #endif
