LandscapePatch: Expose weight patch blend mode override in detail panel.

diff --git a/Engine/Plugins/Experimental/LandscapePatch/Source/LandscapePatch/Public/LandscapeTexturePatch.h b/Engine/Plugins/Experimental/LandscapePatch/Source/LandscapePatch/Public/LandscapeTexturePatch.h
--- a/Engine/Plugins/Experimental/LandscapePatch/Source/LandscapePatch/Public/LandscapeTexturePatch.h
+++ b/Engine/Plugins/Experimental/LandscapePatch/Source/LandscapePatch/Public/LandscapeTexturePatch.h
@@ -193,9 +193,9 @@ protected:
 	bool bUseAlphaChannel = false;
 
 	// Can't make TOptional a UPROPERTY, hence these two.
-	UPROPERTY()
+	UPROPERTY(EditAnywhere, Category = WeightPatch)
 	bool bOverrideBlendMode = false;
-	UPROPERTY()
+	UPROPERTY(EditAnywhere, Category = WeightPatch, meta = (EditConditionHides, EditCondition = "bOverrideBlendMode"))
 	ELandscapeTexturePatchBlendMode OverrideBlendMode = ELandscapeTexturePatchBlendMode::AlphaBlend;
 
 	// TODO: We could support having different per-layer falloff modes and falloff amounts as well, as
