LandscapePatch: Fix internal render target alpha not being preserved in patches. Fix crash on Mac with an older shader compiler that requires shader output to have sufficient channels. Fix height patch reinitialization not interpolating heights correctly. Remove full input copy on patches (only copy the needed portion).

diff --git a/Engine/Plugins/Experimental/LandscapePatch/Shaders/Private/LandscapeTexturePatchPS.usf b/Engine/Plugins/Experimental/LandscapePatch/Shaders/Private/LandscapeTexturePatchPS.usf
--- a/Engine/Plugins/Experimental/LandscapePatch/Shaders/Private/LandscapeTexturePatchPS.usf
+++ b/Engine/Plugins/Experimental/LandscapePatch/Shaders/Private/LandscapeTexturePatchPS.usf
@@ -84,6 +84,7 @@ float GetFalloffAlpha(float FalloffWorldMargin, float2 PatchWorldDimensions, flo
 
 #if APPLY_HEIGHT_PATCH
 Texture2D<float4> InSourceHeightmap;
+uint2 InSourceHeightmapOffset;
 Texture2D<float4> InHeightPatch;
 SamplerState InHeightPatchSampler;
 float4x4 InHeightmapToPatch;
@@ -101,7 +102,7 @@ uint InBlendMode;
 // A combination of flags, whose positions are set in the corresponding cpp file.
 uint InFlags;
 
-void ApplyLandscapeTextureHeightPatch(in float4 SVPos : SV_POSITION, out float2 OutColor : SV_Target0)
+void ApplyLandscapeTextureHeightPatch(in float4 SVPos : SV_POSITION, out float4 OutColor : SV_Target0)
 {
 	// Unpack our flags. The flag choices are set in FApplyLandscapeTextureHeightPatchPS::ModifyCompilationEnvironment
 	// so that they match to the cpp file
@@ -120,7 +121,7 @@ void ApplyLandscapeTextureHeightPatch(in float4 SVPos : SV_POSITION, out float2
 	
 	float PatchSignedHeight = InHeightScale * (PatchStoredHeight - InZeroInEncoding) + InHeightOffset;
 	
-	int2 HeightmapCoordinates = floor(SVPos.xy);
+	int2 HeightmapCoordinates = floor(SVPos.xy) - InSourceHeightmapOffset;
 	float4 CurrentPackedHeight = InSourceHeightmap.Load(int3(HeightmapCoordinates, 0));
 	float CurrentHeight = UnpackHeight(CurrentPackedHeight.xy);
 	
@@ -148,7 +149,7 @@ void ApplyLandscapeTextureHeightPatch(in float4 SVPos : SV_POSITION, out float2
 			break;
 	}
 	
-	OutColor = PackHeight(NewHeight);
+	OutColor = float4(PackHeight(NewHeight), CurrentPackedHeight.z, CurrentPackedHeight.w);
 }
 #endif // APPLY_HEIGHT_PATCH
 
@@ -156,10 +157,11 @@ void ApplyLandscapeTextureHeightPatch(in float4 SVPos : SV_POSITION, out float2
 Texture2D<float4> InHeightmap;
 float InHeightOffset;
 
-void ApplyOffsetToHeightmap(in float4 SVPos : SV_POSITION, out float2 OutColor : SV_Target0)
+void ApplyOffsetToHeightmap(in float4 SVPos : SV_POSITION, out float4 OutColor : SV_Target0)
 {
-	float StoredHeight = UnpackHeight(InHeightmap.Load(int3(SVPos.xy, 0)).xy);
-	OutColor = PackHeight(StoredHeight + InHeightOffset);
+	float4 StoredVector = InHeightmap.Load(int3(SVPos.xy, 0));
+	float StoredHeight = UnpackHeight(StoredVector.xy);
+	OutColor = float4(PackHeight(StoredHeight + InHeightOffset), StoredVector.z, StoredVector.w);
 }
 #endif // OFFSET_HEIGHT_PATCH
 
@@ -169,11 +171,13 @@ float InZeroInEncoding;
 float InHeightScale;
 float InHeightOffset;
 
-void ConvertToNativeLandscapePatch(in float4 SVPos : SV_POSITION, out float2 OutColor : SV_Target0)
+void ConvertToNativeLandscapePatch(in float4 SVPos : SV_POSITION, out float4 OutColor : SV_Target0)
 {
-	float StoredHeight = InHeightmap.Load(int3(SVPos.xy, 0)).x;
+	float4 StoredVector = InHeightmap.Load(int3(SVPos.xy, 0));
+	float StoredHeight = StoredVector.x;
 	// Apply the same kinds of transformations that we would apply for alpha blend mode in ApplyLandscapeTextureHeightPatch
-	OutColor = PackHeight(InHeightScale * (StoredHeight - InZeroInEncoding) + InHeightOffset + LANDSCAPE_MID_VALUE);
+	float InterpretedHeight = InHeightScale * (StoredHeight - InZeroInEncoding) + InHeightOffset + LANDSCAPE_MID_VALUE;
+	OutColor = float4(PackHeight(InterpretedHeight), 0, StoredVector.w);
 }
 #endif // CONVERT_TO_NATIVE_LANDSCAPE_PATCH
 
@@ -183,17 +187,20 @@ float InZeroInEncoding;
 float InHeightScale;
 float InHeightOffset;
 
-void ConvertBackFromNativeLandscapePatch(in float4 SVPos : SV_POSITION, out float OutColor : SV_Target0)
+void ConvertBackFromNativeLandscapePatch(in float4 SVPos : SV_POSITION, out float4 OutColor : SV_Target0)
 {
-	float StoredHeight = UnpackHeight(InHeightmap.Load(int3(SVPos.xy, 0)).xy);
+	float4 StoredVector = InHeightmap.Load(int3(SVPos.xy, 0));
+	float StoredHeight = UnpackHeight(StoredVector.xy);
 	// Undo the transformation we do in ConvertToNativeLandscapePatch
-	OutColor = InHeightScale ? (StoredHeight - LANDSCAPE_MID_VALUE - InHeightOffset) / InHeightScale + InZeroInEncoding
+	float PatchStoredHeight = InHeightScale ? (StoredHeight - LANDSCAPE_MID_VALUE - InHeightOffset) / InHeightScale + InZeroInEncoding
 		: InZeroInEncoding;
+	OutColor = float4(PatchStoredHeight, 0, 0, StoredVector.w);
 }
 #endif // CONVERT_BACK_FROM_NATIVE_LANDSCAPE_PATCH
 
 #if APPLY_WEIGHT_PATCH
 Texture2D<float4> InSourceWeightmap;
+uint2 InSourceWeightmapCoordOffset;
 Texture2D<float4> InWeightPatch;
 SamplerState InWeightPatchSampler;
 float4x4 InWeightmapToPatch;
@@ -222,8 +229,8 @@ void ApplyLandscapeTextureWeightPatch(in float4 SVPos : SV_POSITION, out float O
 	float4 PatchSampledValue = InWeightPatch.Sample(InWeightPatchSampler, PatchUVCoordinates);
 	float PatchWeight = PatchSampledValue.x;
 	
-	int2 WeightmapCoordinates = floor(SVPos.xy);
-	float CurrentWeight = InSourceWeightmap.Load(int3(WeightmapCoordinates, 0)).x;
+	int2 SourceWeightmapCoordinates = floor(SVPos.xy) - InSourceWeightmapCoordOffset;
+	float CurrentWeight = InSourceWeightmap.Load(int3(SourceWeightmapCoordinates, 0)).x;
 	
 	float Alpha = GetFalloffAlpha(InFalloffWorldMargin, InPatchWorldDimensions, PatchUVCoordinates, InEdgeUVDeadBorder, bRectangularFalloff);
 	
@@ -269,6 +276,12 @@ void ReinitializePatch(in float4 SVPos : SV_POSITION, out float4 OutColor : SV_T
 	
 	float2 SourceUVCoordinates = mul(PatchToSourceRotateScale, SVPos.xy) + PatchToSourceTranslate;
 	OutColor = InSource.Sample(InSourceSampler, SourceUVCoordinates);
+
+#if HEIGHT_PATCH
+	// Unpacking and repacking is an easy way to make sure that we interpolate 
+	// between the two bytes of height correctly.
+	OutColor = float4(PackHeight(UnpackHeight(OutColor.xy)), 0, 1);
+#endif
 }
 
 #endif // REINITIALIZE_PATCH
diff --git a/Engine/Plugins/Experimental/LandscapePatch/Source/LandscapePatch/Private/LandscapeTexturePatch.cpp b/Engine/Plugins/Experimental/LandscapePatch/Source/LandscapePatch/Private/LandscapeTexturePatch.cpp
--- a/Engine/Plugins/Experimental/LandscapePatch/Source/LandscapePatch/Private/LandscapeTexturePatch.cpp
+++ b/Engine/Plugins/Experimental/LandscapePatch/Source/LandscapePatch/Private/LandscapeTexturePatch.cpp
@@ -204,12 +204,19 @@ UTextureRenderTarget2D* ULandscapeTexturePatch::ApplyToHeightmap(UTextureRenderT
 			TRefCountPtr<IPooledRenderTarget> DestinationRenderTarget = CreateRenderTarget(InCombinedResult->GetResource()->GetTexture2DRHI(), TEXT("LandscapeTextureHeightPatchOutput"));
 			FRDGTextureRef DestinationTexture = GraphBuilder.RegisterExternalTexture(DestinationRenderTarget);
 
-			// Make a copy of our heightmap input so we can read and write at the same time (needed for blending)
-			FRDGTextureRef InputCopy = GraphBuilder.CreateTexture(DestinationTexture->Desc, TEXT("LandscapeTextureHeightPatchInputCopy"));
+			// Make a copy of the portion of our heightmap input that we're writing to so that we can 
+			// read and write at the same time (needed for blending)
+			FRDGTextureDesc InputCopyDescription = DestinationTexture->Desc;
+			InputCopyDescription.NumMips = 1;
+			InputCopyDescription.Extent = DestinationBounds.Size();
+			FRDGTextureRef InputCopy = GraphBuilder.CreateTexture(InputCopyDescription, TEXT("LandscapeTextureHeightPatchInputCopy"));
 
 			FRHICopyTextureInfo CopyTextureInfo;
+			CopyTextureInfo.SourceMipIndex = 0;
 			CopyTextureInfo.NumMips = 1;
-			CopyTextureInfo.Size = FIntVector(DestinationTexture->Desc.GetSize().X, DestinationTexture->Desc.GetSize().Y, 0);
+			CopyTextureInfo.SourcePosition = FIntVector(DestinationBounds.Min.X, DestinationBounds.Min.Y, 0);
+			CopyTextureInfo.Size = FIntVector(InputCopyDescription.Extent.X, InputCopyDescription.Extent.Y, 0);
+
 			AddCopyTexturePass(GraphBuilder, DestinationTexture, InputCopy, CopyTextureInfo);
 
 			FApplyLandscapeTextureHeightPatchPS::FParameters* ShaderParams =
@@ -224,6 +231,7 @@ UTextureRenderTarget2D* ULandscapeTexturePatch::ApplyToHeightmap(UTextureRenderT
 
 			FRDGTextureSRVRef InputCopySRV = GraphBuilder.CreateSRV(FRDGTextureSRVDesc::CreateForMipLevel(InputCopy, 0));
 			ShaderParams->InSourceHeightmap = InputCopySRV;
+			ShaderParams->InSourceHeightmapOffset = DestinationBounds.Min;
 
 			ShaderParams->RenderTargets[0] = FRenderTargetBinding(DestinationTexture, ERenderTargetLoadAction::ENoAction, /*InMipIndex = */0);
 
@@ -294,21 +302,35 @@ UTextureRenderTarget2D* ULandscapeTexturePatch::ApplyToWeightmap(ULandscapeWeigh
 		return InCombinedResult;
 	}
 
-	ENQUEUE_RENDER_COMMAND(LandscapeTextureHeightPatch)([InCombinedResult, ShaderParamsToCopy, Patch, DestinationBounds](FRHICommandListImmediate& RHICmdList)
+	// Will be used later, once we implement batched merge
+	int32 LandscapeTextureSliceIndex = 0;
+
+	ENQUEUE_RENDER_COMMAND(LandscapeTextureHeightPatch)(
+		[InCombinedResult, LandscapeTextureSliceIndex, ShaderParamsToCopy, Patch, DestinationBounds](FRHICommandListImmediate& RHICmdList)
 		{
 			TRACE_CPUPROFILER_EVENT_SCOPE(LandscapeTextureHeightPatch_Render);
 
 			FRDGBuilder GraphBuilder(RHICmdList, RDG_EVENT_NAME("ApplyTextureHeightPatch"));
 
-			TRefCountPtr<IPooledRenderTarget> DestinationRenderTarget = CreateRenderTarget(InCombinedResult->GetResource()->GetTexture2DRHI(), TEXT("LandscapeTextureWeightPatchOutput"));
+			TRefCountPtr<IPooledRenderTarget> DestinationRenderTarget = CreateRenderTarget(
+				InCombinedResult->GetResource()->GetTexture2DRHI(), TEXT("LandscapeTextureWeightPatchOutput"));
 			FRDGTextureRef DestinationTexture = GraphBuilder.RegisterExternalTexture(DestinationRenderTarget);
 
-			// Make a copy of our heightmap input so we can read and write at the same time (needed for blending)
-			FRDGTextureRef InputCopy = GraphBuilder.CreateTexture(DestinationTexture->Desc, TEXT("LandscapeTextureWeightPatchInputCopy"));
+			// Make a copy of the portion of our weightmap input that we're writing to so that we can 
+			// read and write at the same time (needed for blending)
+			FRDGTextureDesc InputCopyDescription = DestinationTexture->Desc;
+			InputCopyDescription.ArraySize = 1;
+			InputCopyDescription.NumMips = 1;
+			InputCopyDescription.Extent = DestinationBounds.Size();
+			FRDGTextureRef InputCopy = GraphBuilder.CreateTexture(InputCopyDescription, TEXT("LandscapeTextureWeightPatchInputCopy"));
 
 			FRHICopyTextureInfo CopyTextureInfo;
+			CopyTextureInfo.SourceMipIndex = 0;
 			CopyTextureInfo.NumMips = 1;
-			CopyTextureInfo.Size = FIntVector(DestinationTexture->Desc.GetSize().X, DestinationTexture->Desc.GetSize().Y, 0);
+			CopyTextureInfo.SourceSliceIndex = LandscapeTextureSliceIndex;
+			CopyTextureInfo.NumSlices = 1;
+			CopyTextureInfo.SourcePosition = FIntVector(DestinationBounds.Min.X, DestinationBounds.Min.Y, 0);
+			CopyTextureInfo.Size = FIntVector(InputCopyDescription.Extent.X, InputCopyDescription.Extent.Y, 0);
 			AddCopyTexturePass(GraphBuilder, DestinationTexture, InputCopy, CopyTextureInfo);
 
 			FApplyLandscapeTextureWeightPatchPS::FParameters* ShaderParams =
@@ -323,8 +345,10 @@ UTextureRenderTarget2D* ULandscapeTexturePatch::ApplyToWeightmap(ULandscapeWeigh
 
 			FRDGTextureSRVRef InputCopySRV = GraphBuilder.CreateSRV(FRDGTextureSRVDesc::CreateForMipLevel(InputCopy, 0));
 			ShaderParams->InSourceWeightmap = InputCopySRV;
+			ShaderParams->InSourceWeightmapCoordOffset = DestinationBounds.Min;
 
-			ShaderParams->RenderTargets[0] = FRenderTargetBinding(DestinationTexture, ERenderTargetLoadAction::ENoAction, /*InMipIndex = */0);
+			ShaderParams->RenderTargets[0] = FRenderTargetBinding(DestinationTexture, ERenderTargetLoadAction::ENoAction,
+				/*InMipIndex = */0, LandscapeTextureSliceIndex);
 
 			FApplyLandscapeTextureWeightPatchPS::AddToRenderGraph(GraphBuilder, ShaderParams, DestinationBounds);
 
@@ -713,7 +737,7 @@ void ULandscapeTexturePatch::ReinitializeHeight(UTextureRenderTarget2D* InCombin
 			FRDGTextureRef TemporaryDestination = GraphBuilder.CreateTexture(DestinationTexture->Desc, TEXT("LandscapeTextureHeightPatchInputCopy"));
 			HeightmapResampleParams->RenderTargets[0] = FRenderTargetBinding(TemporaryDestination, ERenderTargetLoadAction::ENoAction, /*InMipIndex = */0);
 
-			FReinitializeLandscapePatchPS::AddToRenderGraph(GraphBuilder, HeightmapResampleParams);
+			FReinitializeLandscapePatchPS::AddToRenderGraph(GraphBuilder, HeightmapResampleParams, /*bHeightPatch*/ true);
 
 			FOffsetHeightmapPS::FParameters* OffsetParams = GraphBuilder.AllocParameters<FOffsetHeightmapPS::FParameters>();
 
@@ -727,7 +751,7 @@ void ULandscapeTexturePatch::ReinitializeHeight(UTextureRenderTarget2D* InCombin
 		else
 		{
 			HeightmapResampleParams->RenderTargets[0] = FRenderTargetBinding(DestinationTexture, ERenderTargetLoadAction::ENoAction, /*InMipIndex = */0);
-			FReinitializeLandscapePatchPS::AddToRenderGraph(GraphBuilder, HeightmapResampleParams);
+			FReinitializeLandscapePatchPS::AddToRenderGraph(GraphBuilder, HeightmapResampleParams, /*bHeightPatch*/ true);
 		}
 
 		GraphBuilder.Execute();
@@ -824,13 +848,14 @@ void ULandscapeTexturePatch::ReinitializeWeightPatch(ULandscapeWeightPatchTextur
 		FRDGTextureRef SourceTexture = GraphBuilder.RegisterExternalTexture(CreateRenderTarget(Source->GetTexture2DRHI(), TEXT("ReinitializationSource")));
 		FRDGTextureSRVRef SourceSRV = GraphBuilder.CreateSRV(FRDGTextureSRVDesc::CreateForMipLevel(SourceTexture, 0));
 		ShaderParams->InSource = SourceSRV;
+
 		ShaderParams->InSourceSampler = TStaticSamplerState<SF_Bilinear, AM_Clamp, AM_Clamp>::GetRHI();
 
 		ShaderParams->InPatchToSource = PatchToSource;
 
 		FRDGTextureRef DestinationTexture = GraphBuilder.RegisterExternalTexture(CreateRenderTarget(Destination->GetTexture2DRHI(), TEXT("ReinitializationDestination")));
 		ShaderParams->RenderTargets[0] = FRenderTargetBinding(DestinationTexture, ERenderTargetLoadAction::ENoAction, /*InMipIndex = */0);
-		FReinitializeLandscapePatchPS::AddToRenderGraph(GraphBuilder, ShaderParams);
+		FReinitializeLandscapePatchPS::AddToRenderGraph(GraphBuilder, ShaderParams, /*bHeightPatch*/ false);
 
 		GraphBuilder.Execute();
 	});
diff --git a/Engine/Plugins/Experimental/LandscapePatch/Source/LandscapePatch/Private/LandscapeTexturePatchPS.cpp b/Engine/Plugins/Experimental/LandscapePatch/Source/LandscapePatch/Private/LandscapeTexturePatchPS.cpp
--- a/Engine/Plugins/Experimental/LandscapePatch/Source/LandscapePatch/Private/LandscapeTexturePatchPS.cpp
+++ b/Engine/Plugins/Experimental/LandscapePatch/Source/LandscapePatch/Private/LandscapeTexturePatchPS.cpp
@@ -216,17 +216,20 @@ void FReinitializeLandscapePatchPS::ModifyCompilationEnvironment(const FGlobalSh
 	OutEnvironment.SetDefine(TEXT("REINITIALIZE_PATCH"), 1);
 }
 
-void FReinitializeLandscapePatchPS::AddToRenderGraph(FRDGBuilder& GraphBuilder, FParameters* InParameters)
+void FReinitializeLandscapePatchPS::AddToRenderGraph(FRDGBuilder& GraphBuilder, FParameters* InParameters, bool bHeightPatch)
 {
 	FGlobalShaderMap* ShaderMap = GetGlobalShaderMap(GMaxRHIFeatureLevel);
-	TShaderMapRef<FReinitializeLandscapePatchPS> PixelShader(ShaderMap);
+
+	FReinitializeLandscapePatchPS::FPermutationDomain PermutationDomain;
+	PermutationDomain.Set<FReinitializeLandscapePatchPS::FHeightPatch>(bHeightPatch);
+	TShaderMapRef<FReinitializeLandscapePatchPS> PixelShader(ShaderMap, PermutationDomain);
 
 	FIntVector DestinationSize = InParameters->RenderTargets[0].GetTexture()->Desc.GetSize();
 
 	FPixelShaderUtils::AddFullscreenPass(
 		GraphBuilder,
 		ShaderMap,
-		RDG_EVENT_NAME("LandscapeTextureHeightPatch"),
+		RDG_EVENT_NAME("ReinitializeLandscapeTexturePatch"),
 		PixelShader,
 		InParameters,
 		FIntRect(0, 0, DestinationSize.X, DestinationSize.Y));
diff --git a/Engine/Plugins/Experimental/LandscapePatch/Source/LandscapePatch/Public/LandscapeTexturePatchPS.h b/Engine/Plugins/Experimental/LandscapePatch/Source/LandscapePatch/Public/LandscapeTexturePatchPS.h
--- a/Engine/Plugins/Experimental/LandscapePatch/Source/LandscapePatch/Public/LandscapeTexturePatchPS.h
+++ b/Engine/Plugins/Experimental/LandscapePatch/Source/LandscapePatch/Public/LandscapeTexturePatchPS.h
@@ -75,6 +75,11 @@ public:
 
 	BEGIN_SHADER_PARAMETER_STRUCT(FParameters, )
 		SHADER_PARAMETER_RDG_TEXTURE_SRV(Texture2D<float4>, InSourceHeightmap)
+		// Offset of the source heightmap relative to the 0,0 location in the destination heightmap, because
+		// the source is likely to be a copied region from some inner part of the destination. This is basically
+		// a SourceHeightmapToDestinationHeightmap coordinate transformation, except that it is always a simple 
+		// integer translation.
+		SHADER_PARAMETER(FIntPoint, InSourceHeightmapOffset)
 		SHADER_PARAMETER_RDG_TEXTURE_SRV(Texture2D<float4>, InHeightPatch)
 		SHADER_PARAMETER_SAMPLER(SamplerState, InHeightPatchSampler)
 		SHADER_PARAMETER(FMatrix44f, InHeightmapToPatch)
@@ -223,6 +228,11 @@ public:
 
 	BEGIN_SHADER_PARAMETER_STRUCT(FParameters, )
 		SHADER_PARAMETER_RDG_TEXTURE_SRV(Texture2D<float4>, InSourceWeightmap)
+		// Offset of the source weightmap relative to the 0,0 location in the destination weightmap, because
+		// the source is likely to be a copied region from some inner part of the destination. This is basically
+		// a SourceWeightmapToDestinationWeightmap coordinate transformation, except that it is always a simple 
+		// integer translation.
+		SHADER_PARAMETER(FIntPoint, InSourceWeightmapCoordOffset)
 		SHADER_PARAMETER_RDG_TEXTURE_SRV(Texture2D<float4>, InWeightPatch)
 		SHADER_PARAMETER_SAMPLER(SamplerState, InWeightPatchSampler)
 		SHADER_PARAMETER(FMatrix44f, InWeightmapToPatch)
@@ -255,6 +265,9 @@ class FReinitializeLandscapePatchPS : public FGlobalShader
 
 public:
 
+	class FHeightPatch : SHADER_PERMUTATION_BOOL("HEIGHT_PATCH");
+	using FPermutationDomain = TShaderPermutationDomain<FHeightPatch>;
+
 	BEGIN_SHADER_PARAMETER_STRUCT(FParameters, )
 		SHADER_PARAMETER_RDG_TEXTURE_SRV(Texture2D<float4>, InSource)
 		SHADER_PARAMETER_SAMPLER(SamplerState, InSourceSampler)
@@ -266,7 +279,7 @@ public:
 	static bool ShouldCompilePermutation(const FGlobalShaderPermutationParameters& Parameters);
 	static void ModifyCompilationEnvironment(const FGlobalShaderPermutationParameters& InParameters, FShaderCompilerEnvironment& OutEnvironment);
 
-	LANDSCAPEPATCH_API static void AddToRenderGraph(FRDGBuilder& GraphBuilder, FParameters* InParameters);
+	LANDSCAPEPATCH_API static void AddToRenderGraph(FRDGBuilder& GraphBuilder, FParameters* InParameters, bool bHeightPatch);
 };
 
 /**
