LandscapePatch: Expose a falloff mode setter for blueprints. Avoid dirtying in setter methods if nothing changes.

diff --git a/Engine/Plugins/Experimental/LandscapePatch/Source/LandscapePatch/Public/LandscapeTexturePatch.h b/Engine/Plugins/Experimental/LandscapePatch/Source/LandscapePatch/Public/LandscapeTexturePatch.h
--- a/Engine/Plugins/Experimental/LandscapePatch/Source/LandscapePatch/Public/LandscapeTexturePatch.h
+++ b/Engine/Plugins/Experimental/LandscapePatch/Source/LandscapePatch/Public/LandscapeTexturePatch.h
@@ -294,8 +294,21 @@ public:
 	UFUNCTION(BlueprintCallable, Category = LandscapePatch)
 	void SetFalloff(float FalloffIn) 
 	{
-		Modify();
-		Falloff = FalloffIn; 
+		if (Falloff != FalloffIn)
+		{
+			Modify();
+			Falloff = FalloffIn;
+		}
+	}
+
+	UFUNCTION(BlueprintCallable, Category = LandscapePatch)
+	void SetFalloffMode(ELandscapeTexturePatchFalloffMode FalloffModeIn) 
+	{
+		if (FalloffMode != FalloffModeIn)
+		{
+			Modify();
+			FalloffMode = FalloffModeIn;
+		}
 	}
 
 	/**
@@ -303,9 +316,12 @@ public:
 	 */
 	UFUNCTION(BlueprintCallable, Category = LandscapePatch)
 	void SetBlendMode(ELandscapeTexturePatchBlendMode BlendModeIn) 
-	{ 
-		Modify();
-		BlendMode = BlendModeIn; 
+	{
+		if (BlendMode != BlendModeIn)
+		{
+			Modify();
+			BlendMode = BlendModeIn;
+		}
 	}
 
 
@@ -361,9 +377,12 @@ public:
 	 */
 	UFUNCTION(BlueprintCallable, Category = "LandscapePatch")
 	void SetUseAlphaChannelForHeight(bool bUse)
-	{ 
-		Modify();
-		bUseTextureAlphaForHeight = bUse; 
+	{
+		if (bUseTextureAlphaForHeight != bUse)
+		{
+			Modify();
+			bUseTextureAlphaForHeight = bUse;
+		}
 	}
 
 	/**
@@ -373,8 +392,11 @@ public:
 	UFUNCTION(BlueprintCallable, Category = LandscapePatch)
 	void SetHeightEncodingMode(ELandscapeTextureHeightPatchEncoding EncodingMode)
 	{
-		Modify();
-		HeightEncoding = EncodingMode;
+		if (HeightEncoding != EncodingMode)
+		{
+			Modify();
+			HeightEncoding = EncodingMode;
+		}
 	}
 
 	/**
@@ -398,8 +420,11 @@ public:
 	UFUNCTION(BlueprintCallable, Category = "LandscapePatch")
 	void SetZeroHeightMeaning(ELandscapeTextureHeightPatchZeroHeightMeaning ZeroHeightMeaningIn)
 	{ 
-		Modify();
-		ZeroHeightMeaning = ZeroHeightMeaningIn;
+		if (ZeroHeightMeaning != ZeroHeightMeaningIn)
+		{
+			Modify();
+			ZeroHeightMeaning = ZeroHeightMeaningIn;
+		}
 	}
 
 
