LandscapePatch: Add support for batched merge application path. Also update circle patch to not do full landscape copy and fix an issue where visibility could be broken across the landscape if a brush tries to edit a weightmap with the same layer name.

diff --git a/Engine/Plugins/Experimental/LandscapePatch/Shaders/Private/LandscapeCircleHeightPatchPS.usf b/Engine/Plugins/Experimental/LandscapePatch/Shaders/Private/LandscapeCircleHeightPatchPS.usf
--- a/Engine/Plugins/Experimental/LandscapePatch/Shaders/Private/LandscapeCircleHeightPatchPS.usf
+++ b/Engine/Plugins/Experimental/LandscapePatch/Shaders/Private/LandscapeCircleHeightPatchPS.usf
@@ -9,6 +9,7 @@ float3 InCenter;
 float InRadius;
 float InFalloff;
 Texture2D<float4> InSourceTexture;
+uint2 InSourceTextureOffset;
 
 #if CIRCLE_HEIGHT_PATCH
 
@@ -17,13 +18,14 @@ void ApplyLandscapeCircleHeightPatch(in float4 SVPos : SV_POSITION, out float4 O
 	const float KINDA_SMALL_NUMBER = 0.0001;
 	const float ClampedFalloff = max(InFalloff, KINDA_SMALL_NUMBER);
 	
-	int2 TextureCoordinates = floor(SVPos.xy);
-	float4 CurrentPackedHeight = InSourceTexture.Load(int3(TextureCoordinates, 0));
+	int2 DestinationTextureCoordinates = floor(SVPos.xy);
+	int2 SourceTextureCoordinates = DestinationTextureCoordinates - InSourceTextureOffset;
+	float4 CurrentPackedHeight = InSourceTexture.Load(int3(SourceTextureCoordinates, 0));
 	float CurrentHeight = UnpackHeight(CurrentPackedHeight.xy);
 	
 	float PatchHeight = InCenter.z;
 	
-	float Distance = distance(TextureCoordinates, InCenter.xy);
+	float Distance = distance(DestinationTextureCoordinates, InCenter.xy);
 	
 	float Alpha = 1.0;
 	if (Distance > InRadius)
@@ -42,13 +44,17 @@ void ApplyLandscapeCircleHeightPatch(in float4 SVPos : SV_POSITION, out float4 O
 
 void ApplyLandscapeCircleVisibilityPatch(in float4 SVPos : SV_POSITION, out float OutColor : SV_Target0)
 {
-	int2 TextureCoordinates = floor(SVPos.xy);
-	float Distance = distance(TextureCoordinates, InCenter.xy);
+	int2 DestinationTextureCoordinates = floor(SVPos.xy);
+	int2 SourceTextureCoordinates = DestinationTextureCoordinates - InSourceTextureOffset;
+	float Distance = distance(DestinationTextureCoordinates, InCenter.xy);
 	
-	float4 CurrentWeight = InSourceTexture.Load(int3(TextureCoordinates, 0)).x;
+	float4 CurrentWeight = InSourceTexture.Load(int3(SourceTextureCoordinates, 0)).x;
 
 	if (Distance <= InRadius)
 	{
+		// TODO: This should be changed to calculate alpha just like a height patch and write out
+		//  the alpha value, which allows for smoother holes. But needs to be changed in a
+		//  separate CL with a fixup to set the falloff of older assets to 0 to keep them unchanged.
 		OutColor = 255.f;
 	}
 	else
diff --git a/Engine/Plugins/Experimental/LandscapePatch/Source/LandscapePatch/Private/LandscapeCircleHeightPatch.cpp b/Engine/Plugins/Experimental/LandscapePatch/Source/LandscapePatch/Private/LandscapeCircleHeightPatch.cpp
--- a/Engine/Plugins/Experimental/LandscapePatch/Source/LandscapePatch/Private/LandscapeCircleHeightPatch.cpp
+++ b/Engine/Plugins/Experimental/LandscapePatch/Source/LandscapePatch/Private/LandscapeCircleHeightPatch.cpp
@@ -3,14 +3,17 @@
 #include "LandscapeCircleHeightPatch.h"
 
 #include "Engine/TextureRenderTarget2D.h"
+#include "Engine/TextureRenderTarget2DArray.h"
+#include "Landscape.h"
 #include "LandscapeCircleHeightPatchPS.h"
-#include "LandscapeEditTypes.h"
-#include "LandscapeLayerInfoObject.h"
-#include "LandscapePatchManager.h"
-#include "LandscapeProxy.h"
-#include "RenderGraphUtils.h"
-#include "TextureResource.h"
+#include "LandscapeEditResourcesSubsystem.h" // ULandscapeScratchRenderTarget
+#include "LandscapeInfo.h"
+#include "LandscapePatchUtil.h" // GetHeightmapToWorld
+#include "LandscapeUtils.h" // IsVisibilityLayer
+#include "RenderGraphBuilder.h"
+#include "RenderGraphUtils.h" // AddCopyTexturePass
 #include "RenderingThread.h"
+#include "TextureResource.h"
 
 #include UE_INLINE_GENERATED_CPP_BY_NAME(LandscapeCircleHeightPatch)
 
@@ -48,14 +51,108 @@ UTextureRenderTarget2D* ULandscapeCircleHeightPatch::RenderLayer_Native(const FL
 		return InParameters.CombinedResult;
 	}
 
+	ApplyCirclePatch(InParameters.LayerType == ELandscapeToolTargetType::Visibility,
+		[CombinedResult = InParameters.CombinedResult] { return CombinedResult->GetResource()->GetTexture2DRHI(); }, 0, 
+		FIntPoint(InParameters.CombinedResult->SizeX, InParameters.CombinedResult->SizeY),
+		HeightmapCoordsToWorld);
+
+	return InParameters.CombinedResult;
+}
+
+#if WITH_EDITOR
+void ULandscapeCircleHeightPatch::RenderLayer(FRenderParams& InRenderParams)
+{
+	using namespace UE::Landscape::PatchUtil;
+	using namespace UE::Landscape::EditLayers;
+
+	FTransform HeightmapCoordsToWorld = GetHeightmapToWorld(InRenderParams.RenderAreaWorldTransform);
+
+	ULandscapeScratchRenderTarget* LandscapeRT = InRenderParams.MergeRenderContext->GetBlendRenderTargetWrite();
+
+	const bool bIsHeightmapTarget = InRenderParams.MergeRenderContext->IsHeightmapMerge();
+	if (bIsHeightmapTarget)
+	{
+		UTextureRenderTarget2D* OutputToBlendInto = LandscapeRT->TryGetRenderTarget2D();
+		if (ensure(OutputToBlendInto))
+		{
+			// graph builder expects external textures to start as SRV
+			LandscapeRT->TransitionTo(ERHIAccess::SRVMask); 
+
+			ApplyCirclePatch(false, [OutputToBlendInto] { return OutputToBlendInto->GetResource()->GetTexture2DRHI(); },
+				0, InRenderParams.RenderAreaSectionRect.Size(), HeightmapCoordsToWorld);
+		}
+		return;
+	}
+
+	// If we got to here, we're not processing a heightmap, so we only need to do anything if the 
+	//  patch edits visibility.
+	if (!bEditVisibility)
+	{
+		return;
+	}
+
+	UTextureRenderTarget2DArray* TextureArray = LandscapeRT->TryGetRenderTarget2DArray();
+	if (!ensure(TextureArray))
+	{
+		return;
+	}
+
+	int32 NumLayers = InRenderParams.RenderGroupTargetLayerInfos.Num();
+	if (!ensure(TextureArray->Slices == NumLayers))
+	{
+		NumLayers = FMath::Min(NumLayers, TextureArray->Slices);
+	}
+
+	for (int32 LayerIndex = 0; LayerIndex < NumLayers; ++LayerIndex)
+	{
+		if (UE::Landscape::IsVisibilityLayer(InRenderParams.RenderGroupTargetLayerInfos[LayerIndex]))
+		{
+			// graph builder expects external textures to start as SRV
+			LandscapeRT->TransitionTo(ERHIAccess::SRVMask);
+
+			ApplyCirclePatch(true, [TextureArray] { return TextureArray->GetResource()->GetTexture2DArrayRHI(); }, 
+				LayerIndex, InRenderParams.RenderAreaSectionRect.Size(), HeightmapCoordsToWorld);
+		}
+	}
+}
+#endif // WITH_EDITOR
+
+void ULandscapeCircleHeightPatch::ApplyCirclePatch(bool bIsVisibilityLayer,
+	const TFunction<FRHITexture* ()>& RenderThreadLandscapeTextureGetter, int32 LandscapeTextureSliceIndex,
+	const FIntPoint& DestinationResolution, const FTransform & HeightmapCoordsToWorld)
+{
+	if (bEditVisibility != bIsVisibilityLayer)
+	{
+		return;
+	}
+
 	double ToHeightmapRadiusScale = GetComponentTransform().GetScale3D().X / HeightmapCoordsToWorld.GetScale3D().X;
 	FVector3d CircleCenterWorld = GetComponentTransform().GetTranslation();
 	FVector3d CenterInHeightmapCoordinates = HeightmapCoordsToWorld.InverseTransformPosition(CircleCenterWorld);
 	float RadiusAdjustment = bExclusiveRadius ? 0 : 1;
 	float HeightmapRadius = Radius * ToHeightmapRadiusScale + RadiusAdjustment;
+	// TODO: This is incorrect, should not have radius adjustment here. However, need to change in a separate CL
+	//  so that we can add a fixup to leave older assets unchanged.
 	float HeightmapFalloff = Falloff * ToHeightmapRadiusScale + RadiusAdjustment;
 
-	ENQUEUE_RENDER_COMMAND(LandscapeCircleHeightPatch)([CombinedResult = InParameters.CombinedResult, CenterInHeightmapCoordinates, HeightmapRadius, HeightmapFalloff, bEditVisibility = bEditVisibility](FRHICommandListImmediate& RHICmdList)
+	FIntRect DestinationBounds(
+		FMath::Clamp(FMath::Floor(CenterInHeightmapCoordinates.X - HeightmapRadius - HeightmapFalloff), 
+			0, DestinationResolution.X),
+		FMath::Clamp(FMath::Floor(CenterInHeightmapCoordinates.Y - HeightmapRadius - HeightmapFalloff), 
+			0, DestinationResolution.Y),
+		FMath::Clamp(FMath::CeilToInt(CenterInHeightmapCoordinates.X + HeightmapRadius + HeightmapFalloff) + 1, 
+			0, DestinationResolution.X),
+		FMath::Clamp(FMath::CeilToInt(CenterInHeightmapCoordinates.Y + HeightmapRadius + HeightmapFalloff) + 1, 
+			0, DestinationResolution.Y));
+
+	if (DestinationBounds.Area() <= 0)
+	{
+		// Must be outside the landscape
+		return;
+	}
+
+	ENQUEUE_RENDER_COMMAND(LandscapeCircleHeightPatch)([RenderThreadLandscapeTextureGetter, LandscapeTextureSliceIndex, DestinationBounds,
+		CenterInHeightmapCoordinates, HeightmapRadius, HeightmapFalloff, bEditVisibility = bEditVisibility](FRHICommandListImmediate& RHICmdList)
 	{
 		TRACE_CPUPROFILER_EVENT_SCOPE(LandscapeCircleHeightPatch);
 
@@ -64,16 +161,27 @@ UTextureRenderTarget2D* ULandscapeCircleHeightPatch::RenderLayer_Native(const FL
 
 		FRDGBuilder GraphBuilder(RHICmdList, RDG_EVENT_NAME("ApplyLandscapeCirclePatch"));
 		
-		TRefCountPtr<IPooledRenderTarget> RenderTarget = CreateRenderTarget(CombinedResult->GetResource()->GetTexture2DRHI(), OutputName);
+		TRefCountPtr<IPooledRenderTarget> RenderTarget = CreateRenderTarget(RenderThreadLandscapeTextureGetter(), OutputName);
 		FRDGTextureRef DestinationTexture = GraphBuilder.RegisterExternalTexture(RenderTarget);
 
-		// Make a copy of our heightmap input
-		FRDGTextureRef InputCopy = GraphBuilder.CreateTexture(DestinationTexture->Desc, InputCopyName);
+		// Make a copy of the portion of our weightmap input that we're writing to so that we can 
+		//  read and write at the same time (needed for blending)
+		FRDGTextureDesc InputCopyDescription = DestinationTexture->Desc;
+		InputCopyDescription.Dimension = ETextureDimension::Texture2D;
+		InputCopyDescription.ArraySize = 1;
+		InputCopyDescription.NumMips = 1;
+		InputCopyDescription.Extent = DestinationBounds.Size();
+		FRDGTextureRef InputCopy = GraphBuilder.CreateTexture(InputCopyDescription, InputCopyName);
 
 		FRHICopyTextureInfo CopyTextureInfo;
+		CopyTextureInfo.SourceMipIndex = 0;
 		CopyTextureInfo.NumMips = 1;
-		CopyTextureInfo.Size = FIntVector(DestinationTexture->Desc.GetSize().X, DestinationTexture->Desc.GetSize().Y, 0);
+		CopyTextureInfo.SourceSliceIndex = LandscapeTextureSliceIndex;
+		CopyTextureInfo.NumSlices = 1;
+		CopyTextureInfo.SourcePosition = FIntVector(DestinationBounds.Min.X, DestinationBounds.Min.Y, 0);
+		CopyTextureInfo.Size = FIntVector(InputCopyDescription.Extent.X, InputCopyDescription.Extent.Y, 0);
 		AddCopyTexturePass(GraphBuilder, DestinationTexture, InputCopy, CopyTextureInfo);
+
 		FRDGTextureSRVRef InputCopySRV = GraphBuilder.CreateSRV(FRDGTextureSRVDesc::CreateForMipLevel(InputCopy, 0));
 
 		FLandscapeCircleHeightPatchPS::FParameters* ShaderParams = GraphBuilder.AllocParameters<FLandscapeCircleHeightPatchPS::FParameters>();
@@ -81,21 +189,72 @@ UTextureRenderTarget2D* ULandscapeCircleHeightPatch::RenderLayer_Native(const FL
 		ShaderParams->InRadius = HeightmapRadius;
 		ShaderParams->InFalloff = HeightmapFalloff;
 		ShaderParams->InSourceTexture = InputCopySRV;
-		ShaderParams->RenderTargets[0] = FRenderTargetBinding(DestinationTexture, ERenderTargetLoadAction::ENoAction, /*InMipIndex = */0);
+		ShaderParams->InSourceTextureOffset = DestinationBounds.Min;
+		ShaderParams->RenderTargets[0] = FRenderTargetBinding(DestinationTexture, ERenderTargetLoadAction::ENoAction, 
+			/*InMipIndex = */0, LandscapeTextureSliceIndex);
 
 		if (bEditVisibility)
 		{
-			FLandscapeCircleVisibilityPatchPS::AddToRenderGraph(GraphBuilder, ShaderParams);
+			FLandscapeCircleVisibilityPatchPS::AddToRenderGraph(GraphBuilder, ShaderParams, DestinationBounds);
 		}
 		else
 		{
-			FLandscapeCircleHeightPatchPS::AddToRenderGraph(GraphBuilder, ShaderParams);
+			FLandscapeCircleHeightPatchPS::AddToRenderGraph(GraphBuilder, ShaderParams, DestinationBounds);
 		}
 
 		GraphBuilder.Execute();
 	});
+}
 
-	return InParameters.CombinedResult;
+#if WITH_EDITOR
+void ULandscapeCircleHeightPatch::GetRendererStateInfo(const ULandscapeInfo* InLandscapeInfo,
+	FEditLayerTargetTypeState& OutSupported,
+	FEditLayerTargetTypeState& OutEnabled,
+	TArray<TSet<FName>>& OutRenderGroups) const
+{
+	OutSupported.AddTargetType(bEditVisibility ? ELandscapeToolTargetType::Visibility
+		: ELandscapeToolTargetType::Heightmap);
+
+	if (IsEnabled())
+	{
+		OutEnabled = OutSupported;
+	}
+}
+
+FString ULandscapeCircleHeightPatch::GetEditLayerRendererDebugName() const
+{
+	return FString::Printf(TEXT("ULandscapeCircleHeightPatch:%s"), *GetFullName());
+}
+
+TArray<UE::Landscape::EditLayers::FEditLayerRenderItem> ULandscapeCircleHeightPatch::GetRenderItems(const ULandscapeInfo* InLandscapeInfo) const
+{
+	using namespace UE::Landscape::EditLayers;
+
+	TArray<FEditLayerRenderItem> AffectedAreas;
+
+	FTransform ComponentTransform = this->GetComponentToWorld();
+
+	// Figure out the extents of the patch. It will be radius + falloff + an adjustment if we're
+	//  trying to make the whole circle lie flat. The adjustment will be the size of one landscape
+	//  quad, but to be safe we'll make it two quads in each direction.
+	FVector3d LandscapeScale = InLandscapeInfo->LandscapeActor.IsValid() ?
+		InLandscapeInfo->LandscapeActor->GetActorTransform().GetScale3D() : FVector3d::Zero();
+	FVector2D Extents(2 * FMath::Max(LandscapeScale.X, LandscapeScale.Y) + Radius + Falloff);
+
+	FOOBox2D PatchArea(ComponentTransform, Extents);
+
+	FInputWorldArea InputWorldArea = FInputWorldArea::CreateOOBox(PatchArea);
+	FOutputWorldArea OutputWorldArea = FOutputWorldArea::CreateOOBox(PatchArea);
+
+	FEditLayerTargetTypeState TargetInfo(bEditVisibility ? ELandscapeToolTargetTypeFlags::Visibility
+		: ELandscapeToolTargetTypeFlags::Heightmap);
+	FEditLayerRenderItem Item(TargetInfo, InputWorldArea, OutputWorldArea,
+		// ModifyExistingWeightmapsOnly: we want the patch to allocate the visibility layer if it needs to, in its region.
+		false);
+	AffectedAreas.Add(Item);
+
+	return AffectedAreas;
 }
+#endif // WITH_EDITOR
 
 #undef LOCTEXT_NAMESPACE
diff --git a/Engine/Plugins/Experimental/LandscapePatch/Source/LandscapePatch/Private/LandscapeCircleHeightPatchPS.cpp b/Engine/Plugins/Experimental/LandscapePatch/Source/LandscapePatch/Private/LandscapeCircleHeightPatchPS.cpp
--- a/Engine/Plugins/Experimental/LandscapePatch/Source/LandscapePatch/Private/LandscapeCircleHeightPatchPS.cpp
+++ b/Engine/Plugins/Experimental/LandscapePatch/Source/LandscapePatch/Private/LandscapeCircleHeightPatchPS.cpp
@@ -18,20 +18,18 @@ void FLandscapeCircleHeightPatchPS::ModifyCompilationEnvironment(const FGlobalSh
 	OutEnvironment.SetDefine(TEXT("CIRCLE_HEIGHT_PATCH"), 1);
 }
 
-void FLandscapeCircleHeightPatchPS::AddToRenderGraph(FRDGBuilder& GraphBuilder, FParameters* InParameters)
+void FLandscapeCircleHeightPatchPS::AddToRenderGraph(FRDGBuilder& GraphBuilder, FParameters* Parameters, const FIntRect& DestinationBounds)
 {
 	FGlobalShaderMap* ShaderMap = GetGlobalShaderMap(GMaxRHIFeatureLevel);
 	TShaderMapRef<FLandscapeCircleHeightPatchPS> PixelShader(ShaderMap);
 
-	FIntVector TextureSize = InParameters->InSourceTexture->Desc.Texture->Desc.GetSize();
-
 	FPixelShaderUtils::AddFullscreenPass(
 		GraphBuilder,
 		ShaderMap,
 		RDG_EVENT_NAME("LandmassCircleHeightPatch"),
 		PixelShader,
-		InParameters,
-		FIntRect(0, 0, TextureSize.X, TextureSize.Y));
+		Parameters,
+		DestinationBounds);
 }
 
 void FLandscapeCircleVisibilityPatchPS::ModifyCompilationEnvironment(const FGlobalShaderPermutationParameters& Parameters, FShaderCompilerEnvironment& OutEnvironment)
@@ -39,20 +37,18 @@ void FLandscapeCircleVisibilityPatchPS::ModifyCompilationEnvironment(const FGlob
 	OutEnvironment.SetDefine(TEXT("CIRCLE_VISIBILITY_PATCH"), 1);
 }
 
-void FLandscapeCircleVisibilityPatchPS::AddToRenderGraph(FRDGBuilder& GraphBuilder, FParameters* InParameters)
+void FLandscapeCircleVisibilityPatchPS::AddToRenderGraph(FRDGBuilder& GraphBuilder, FParameters* Parameters, const FIntRect& DestinationBounds)
 {
 	FGlobalShaderMap* ShaderMap = GetGlobalShaderMap(GMaxRHIFeatureLevel);
 	TShaderMapRef<FLandscapeCircleVisibilityPatchPS> PixelShader(ShaderMap);
 
-	FIntVector TextureSize = InParameters->InSourceTexture->Desc.Texture->Desc.GetSize();
-
 	FPixelShaderUtils::AddFullscreenPass(
 		GraphBuilder,
 		ShaderMap,
 		RDG_EVENT_NAME("LandmassCircleVisibilityPatch"),
 		PixelShader,
-		InParameters,
-		FIntRect(0, 0, TextureSize.X, TextureSize.Y));
+		Parameters,
+		DestinationBounds);
 }
 
 IMPLEMENT_GLOBAL_SHADER(FLandscapeCircleHeightPatchPS, "/Plugin/LandscapePatch/Private/LandscapeCircleHeightPatchPS.usf", "ApplyLandscapeCircleHeightPatch", SF_Pixel);
diff --git a/Engine/Plugins/Experimental/LandscapePatch/Source/LandscapePatch/Private/LandscapePatchEditLayer.cpp b/Engine/Plugins/Experimental/LandscapePatch/Source/LandscapePatch/Private/LandscapePatchEditLayer.cpp
--- a/Engine/Plugins/Experimental/LandscapePatch/Source/LandscapePatch/Private/LandscapePatchEditLayer.cpp
+++ b/Engine/Plugins/Experimental/LandscapePatch/Source/LandscapePatch/Private/LandscapePatchEditLayer.cpp
@@ -296,6 +296,30 @@ bool ULandscapePatchEditLayer::SupportsTargetType(ELandscapeToolTargetType InTyp
 }
 
 #if WITH_EDITOR
+// Called in batched merge path to apply the patches
+TArray<UE::Landscape::EditLayers::FEditLayerRendererState> ULandscapePatchEditLayer::GetEditLayerRendererStates(const ULandscapeInfo* InLandscapeInfo, bool bInSkipBrush)
+{
+	UpdatePatchListIfDirty();
+
+	TArray<FEditLayerRendererState> RendererStates;
+	RendererStates.Reserve(RegisteredPatches.Num());
+	for (TSoftObjectPtr<ULandscapePatchComponent>& PatchSoft : RegisteredPatches)
+	{
+		ULandscapePatchComponent* Patch = PatchSoft.Get();
+		if (!Patch)
+		{
+			continue;
+		}
+
+		FEditLayerRendererState& RendererState = RendererStates.Emplace_GetRef(Patch, InLandscapeInfo);
+		if (bInSkipBrush || !Patch->IsEnabled())
+		{
+			RendererState.DisableTargetTypeMask(ELandscapeToolTargetTypeFlags::All);
+		}
+	}
+
+	return RendererStates;
+}
 
 void ULandscapePatchEditLayer::PostLoad()
 {
@@ -381,6 +405,7 @@ void ULandscapePatchEditLayer::InitializeAsBlueprintBrush(const FTransform& InLa
 	HeightmapCoordsToWorld = UE::Landscape::PatchUtil::GetHeightmapToWorld(InLandscapeTransform);
 }
 
+// Called in global merge (legacy) mode to apply the patches
 UTextureRenderTarget2D* ULandscapePatchEditLayer::RenderLayerAsBlueprintBrush(const FLandscapeBrushParameters& InParameters)
 {
 	UpdatePatchListIfDirty();
diff --git a/Engine/Plugins/Experimental/LandscapePatch/Source/LandscapePatch/Private/LandscapePatchManager.cpp b/Engine/Plugins/Experimental/LandscapePatch/Source/LandscapePatch/Private/LandscapePatchManager.cpp
--- a/Engine/Plugins/Experimental/LandscapePatch/Source/LandscapePatch/Private/LandscapePatchManager.cpp
+++ b/Engine/Plugins/Experimental/LandscapePatch/Source/LandscapePatch/Private/LandscapePatchManager.cpp
@@ -225,18 +225,33 @@ UTextureRenderTarget2D* ALandscapePatchManager::RenderLayer_Native(const FLandsc
 }
 
 #if WITH_EDITOR
-TArray<UE::Landscape::EditLayers::FEditLayerRenderItem> ALandscapePatchManager::GetRenderItems(const ULandscapeInfo* InLandscapeInfo) const
+// Called in batched merge path to apply the patches
+TArray<UE::Landscape::EditLayers::FEditLayerRendererState> ALandscapePatchManager::GetEditLayerRendererStates(const ULandscapeInfo* InLandscapeInfo, bool bInSkipBrush)
 {
-	using namespace UE::Landscape::EditLayers;
+	using namespace LandscapePatchManagerLocals;
 
-	TArray<FEditLayerRenderItem> RenderItems = Super::GetRenderItems(InLandscapeInfo);
-	check(RenderItems.Num() == 1);
+	FilterLegacyRegisteredPatches(PatchComponents, PatchToIndex, this);
 
-	// For now, this brush requires the entire landscape to be loaded to work deterministically, so we force the input area to be "infinite" : 
-	RenderItems[0].SetInputWorldArea(FInputWorldArea::CreateInfinite());
+	TArray<FEditLayerRendererState> RendererStates;
+	RendererStates.Reserve(PatchComponents.Num());
+	for (TSoftObjectPtr<ULandscapePatchComponent>& PatchSoft : PatchComponents)
+	{
+		ULandscapePatchComponent* Patch = PatchSoft.Get();
+		if (!Patch)
+		{
+			continue;
+		}
 
-	return RenderItems;
+		FEditLayerRendererState& RendererState = RendererStates.Emplace_GetRef(Patch, InLandscapeInfo);
+		if (bInSkipBrush || !Patch->IsEnabled())
+		{
+			RendererState.DisableTargetTypeMask(ELandscapeToolTargetTypeFlags::All);
+		}
+	}
+
+	return RendererStates;
 }
+
 #endif // WITH_EDITOR
 
 void ALandscapePatchManager::SetTargetLandscape(ALandscape* InTargetLandscape)
diff --git a/Engine/Plugins/Experimental/LandscapePatch/Source/LandscapePatch/Private/LandscapeTexturePatch.cpp b/Engine/Plugins/Experimental/LandscapePatch/Source/LandscapePatch/Private/LandscapeTexturePatch.cpp
--- a/Engine/Plugins/Experimental/LandscapePatch/Source/LandscapePatch/Private/LandscapeTexturePatch.cpp
+++ b/Engine/Plugins/Experimental/LandscapePatch/Source/LandscapePatch/Private/LandscapeTexturePatch.cpp
@@ -2,13 +2,15 @@
 
 #include "LandscapeTexturePatch.h"
 
+#include "Engine/TextureRenderTarget2DArray.h"
 #include "Engine/World.h"
 #include "Landscape.h"
 #include "LandscapeDataAccess.h"
-#include "LandscapeLayerInfoObject.h" // VisibilityLayer->LayerName
+#include "LandscapeEditResourcesSubsystem.h" // ULandscapeScratchRenderTarget
 #include "LandscapePatchLogging.h"
 #include "LandscapePatchManager.h"
 #include "LandscapePatchUtil.h" // CopyTextureOnRenderThread
+#include "LandscapeUtils.h" // IsVisibilityLayer
 #include "MathUtil.h"
 #include "RHIStaticStates.h"
 #include "RenderGraphUtils.h"
@@ -73,10 +75,136 @@ namespace LandscapeTexturePatchLocals
 			TextureAsset = nullptr;
 		}
 	}
+
+	// TODO: The way we currently do intialization is a bit of a hack in that we actually request to do
+	//  a landscape update but we read instead of writing. In batched merge, this might not always work
+	//  properly because a patch might be at the edge of a rendered batch, and thus only have part of it
+	//  be initialized properly. The proper way to do reinitialization would be to use a special function
+	//  to render the relevant part of the landscape directly to the patch. We should do this at some point,
+	//  but it is not high priority because reinitialization does not currently seem to be commonly used.
+	// @param PatchToHeightmapUVs This is expected to be a usual math matrix by this point, not Unreal's transposed one
+	void DoReinitializationOverlapCheck(FMatrix44f& PatchToHeightmapUVs, int32 PatchTextureSizeX, int32 PatchTextureSizeY)
+	{
+		auto IsInsideHeightmap = [&PatchToHeightmapUVs](int32 X, int32 Y)->bool
+		{
+			float U = PatchToHeightmapUVs.M[0][0] * X + PatchToHeightmapUVs.M[0][1] * Y + PatchToHeightmapUVs.M[0][3];
+			float V = PatchToHeightmapUVs.M[1][0] * X + PatchToHeightmapUVs.M[1][1] * Y + PatchToHeightmapUVs.M[1][3];
+
+			return U >= 0 && U <= 1 && V >= 0 && V <= 1;
+		};
+
+		if (!IsInsideHeightmap(0, 0)
+			|| !IsInsideHeightmap(0, PatchTextureSizeY-1)
+			|| !IsInsideHeightmap(PatchTextureSizeX-1, 0)
+			|| !IsInsideHeightmap(PatchTextureSizeX-1, PatchTextureSizeY-1))
+		{
+			UE_LOG(LogLandscapePatch, Warning, TEXT("ULandscapeTexturePatch::Reinitialize: Part or all of the patch was outside "
+				"a region of landscape being rendered. Reinitialization might not work be fully supported here."));
+		}
+	}
 #endif // WITH_EDITOR
 }
 
 #if WITH_EDITOR
+void ULandscapeTexturePatch::RenderLayer(FRenderParams& InRenderParams)
+{
+	using namespace UE::Landscape::PatchUtil;
+	using namespace UE::Landscape::EditLayers;
+
+	FTransform LandscapeHeightmapToWorld = GetHeightmapToWorld(InRenderParams.RenderAreaWorldTransform);
+
+	ULandscapeScratchRenderTarget* LandscapeScratchRT = InRenderParams.MergeRenderContext->GetBlendRenderTargetWrite();
+
+	const bool bIsHeightmapTarget = InRenderParams.MergeRenderContext->IsHeightmapMerge();
+	if (bIsHeightmapTarget)
+	{
+		UTextureRenderTarget2D* CurrentData = LandscapeScratchRT->TryGetRenderTarget2D();
+		if (!ensure(CurrentData))
+		{
+			return;
+		}
+
+		// graph builder expects external textures to start as SRV
+		LandscapeScratchRT->TransitionTo(ERHIAccess::SRVMask);
+		
+		if (bReinitializeHeightOnNextRender)
+		{
+			bReinitializeHeightOnNextRender = false;
+			ReinitializeHeight(CurrentData, LandscapeHeightmapToWorld);
+			return;
+		}
+		else
+		{
+			ApplyToHeightmap(CurrentData, GetHeightmapToWorld(InRenderParams.RenderAreaWorldTransform));
+			return;
+		}
+	}
+
+	// If we got to here, we're dealing with weightmaps.
+
+	UTextureRenderTarget2DArray* TextureArray = LandscapeScratchRT->TryGetRenderTarget2DArray();
+	if (!ensure(TextureArray))
+	{
+		return;
+	}
+	
+	// Only need to transition if we get a matching weight patch
+	bool bTransitionedToSRV = false;
+
+	int32 NumLayers = InRenderParams.RenderGroupTargetLayerNames.Num();
+	if (!ensure(TextureArray->Slices == NumLayers))
+	{
+		NumLayers = FMath::Min(NumLayers, TextureArray->Slices);
+	}
+
+	for (int32 LayerIndex = 0; LayerIndex < NumLayers; ++LayerIndex)
+	{
+		bool bIsVisibilityLayer = ensure(LayerIndex < InRenderParams.RenderGroupTargetLayerInfos.Num()) 
+			&& UE::Landscape::IsVisibilityLayer(InRenderParams.RenderGroupTargetLayerInfos[LayerIndex]);
+		
+		// Try to find the weight patch
+		ULandscapeWeightPatchTextureInfo* WeightPatchInfo = nullptr;
+		for (const TObjectPtr<ULandscapeWeightPatchTextureInfo>& WeightPatchEntry : WeightPatches)
+		{
+			if ((bIsVisibilityLayer && WeightPatchEntry->bEditVisibilityLayer)
+				|| (WeightPatchEntry->WeightmapLayerName == InRenderParams.RenderGroupTargetLayerNames[LayerIndex]))
+			{
+				WeightPatchInfo = WeightPatchEntry;
+				break;
+			}
+		}
+
+		if (!WeightPatchInfo)
+		{
+			// Didn't have a patch for this weight layer
+			continue;
+		}
+
+		// graph builder expects external textures to start as SRV
+		if (!bTransitionedToSRV)
+		{
+			LandscapeScratchRT->TransitionTo(ERHIAccess::SRVMask);
+			bTransitionedToSRV = true;
+		}
+
+		if (WeightPatchInfo->bReinitializeOnNextRender)
+		{
+			WeightPatchInfo->bReinitializeOnNextRender = false;
+			ReinitializeWeightPatch(WeightPatchInfo, TextureArray->GetResource(),
+				FIntPoint(TextureArray->SizeX, TextureArray->SizeY), LayerIndex, LandscapeHeightmapToWorld);
+		}
+		else
+		{
+			ApplyToWeightmap(WeightPatchInfo, 
+				[Resource = TextureArray->GetResource()]() { return Resource->GetTexture2DArrayRHI(); },
+				LayerIndex, 
+				InRenderParams.RenderAreaSectionRect.Size(), 
+				GetHeightmapToWorld(InRenderParams.RenderAreaWorldTransform));
+		}
+	}//end for each layer index
+}//end ULandscapeTexturePatch::RenderLayer
+
+// Legacy path, which gets the entire heightmap.
 UTextureRenderTarget2D* ULandscapeTexturePatch::RenderLayer_Native(const FLandscapeBrushParameters& InParameters, 
 	const FTransform& LandscapeHeightmapToWorld)
 {
@@ -132,12 +260,23 @@ UTextureRenderTarget2D* ULandscapeTexturePatch::RenderLayer_Native(const FLandsc
 		if (WeightPatchInfo->bReinitializeOnNextRender)
 		{
 			WeightPatchInfo->bReinitializeOnNextRender = false;
-			ReinitializeWeightPatch(WeightPatchInfo, InParameters.CombinedResult, LandscapeHeightmapToWorld);
+			if (ensure(InParameters.CombinedResult->GetResource()))
+			{
+				ReinitializeWeightPatch(WeightPatchInfo, InParameters.CombinedResult->GetResource(), 
+					FIntPoint(InParameters.CombinedResult->SizeX, InParameters.CombinedResult->SizeY), 
+					-1, // Signifies that this is not a Texture2DArray
+					LandscapeHeightmapToWorld);
+			}
 			return InParameters.CombinedResult;
 		}
 		else
 		{
-			return ApplyToWeightmap(WeightPatchInfo, InParameters.CombinedResult, LandscapeHeightmapToWorld);
+			ApplyToWeightmap(WeightPatchInfo, 
+				[Resource = InParameters.CombinedResult->GetResource()]() { return Resource->GetTexture2DRHI(); },
+				0, // Slice index
+				FIntPoint(InParameters.CombinedResult->SizeX, InParameters.CombinedResult->SizeY),
+				LandscapeHeightmapToWorld);
+			return InParameters.CombinedResult;
 		}
 	}
 }
@@ -189,7 +328,7 @@ UTextureRenderTarget2D* ULandscapeTexturePatch::ApplyToHeightmap(UTextureRenderT
 	FIntRect DestinationBounds;
 	GetHeightShaderParams(LandscapeHeightmapToWorld, FIntPoint(Patch->GetSizeX(), Patch->GetSizeY()), FIntPoint(InCombinedResult->SizeX, InCombinedResult->SizeY), ShaderParamsToCopy, DestinationBounds);
 
-	if (DestinationBounds.IsEmpty())
+	if (DestinationBounds.Area() <= 0)
 	{
 		// Patch must be outside the landscape.
 		return InCombinedResult;
@@ -243,14 +382,15 @@ UTextureRenderTarget2D* ULandscapeTexturePatch::ApplyToHeightmap(UTextureRenderT
 	return InCombinedResult;
 }
 
-UTextureRenderTarget2D* ULandscapeTexturePatch::ApplyToWeightmap(ULandscapeWeightPatchTextureInfo* PatchInfo, 
-	UTextureRenderTarget2D* InCombinedResult, const FTransform& LandscapeHeightmapToWorld)
+void ULandscapeTexturePatch::ApplyToWeightmap(ULandscapeWeightPatchTextureInfo* PatchInfo, 
+	TFunction<FRHITexture*()> RenderThreadLandscapeTextureGetter, int32 LandscapeTextureSliceIndex, 
+	const FIntPoint& LandscapeTextureResolution, const FTransform& LandscapeHeightmapToWorld)
 {
 	using namespace UE::Landscape;
 
 	if (!PatchInfo)
 	{
-		return InCombinedResult;
+		return;
 	}
 
 	UTexture* PatchUObject = nullptr;
@@ -258,7 +398,7 @@ UTextureRenderTarget2D* ULandscapeTexturePatch::ApplyToWeightmap(ULandscapeWeigh
 	switch (PatchInfo->SourceMode)
 	{
 	case ELandscapeTexturePatchSourceMode::None:
-		return InCombinedResult;
+		return;
 	case ELandscapeTexturePatchSourceMode::InternalTexture:
 		PatchUObject = GetWeightPatchInternalTexture(PatchInfo);
 		break;
@@ -269,7 +409,7 @@ UTextureRenderTarget2D* ULandscapeTexturePatch::ApplyToWeightmap(ULandscapeWeigh
 		if (IsValid(PatchInfo->TextureAsset) && !ensureMsgf(PatchInfo->TextureAsset->VirtualTextureStreaming == 0,
 			TEXT("ULandscapeTexturePatch: Virtual textures are not supported")))
 		{
-			return InCombinedResult;
+			return;
 		}
 		PatchUObject = PatchInfo->TextureAsset;
 		break;
@@ -279,13 +419,13 @@ UTextureRenderTarget2D* ULandscapeTexturePatch::ApplyToWeightmap(ULandscapeWeigh
 
 	if (!IsValid(PatchUObject))
 	{
-		return InCombinedResult;
+		return;
 	}
 
 	FTextureResource* Patch = PatchUObject->GetResource();
 	if (!Patch)
 	{
-		return InCombinedResult;
+		return;
 	}
 
 	// Go ahead and pack everything into a copy of the param struct so we don't have to capture everything
@@ -293,69 +433,66 @@ UTextureRenderTarget2D* ULandscapeTexturePatch::ApplyToWeightmap(ULandscapeWeigh
 	FApplyLandscapeTextureWeightPatchPS::FParameters ShaderParamsToCopy;
 	FIntRect DestinationBounds;
 
-	GetWeightShaderParams(LandscapeHeightmapToWorld, FIntPoint(Patch->GetSizeX(), Patch->GetSizeY()), FIntPoint(InCombinedResult->SizeX, InCombinedResult->SizeY),
-		PatchInfo, ShaderParamsToCopy, DestinationBounds);
+	GetWeightShaderParams(LandscapeHeightmapToWorld, FIntPoint(Patch->GetSizeX(), Patch->GetSizeY()), 
+		LandscapeTextureResolution, PatchInfo, ShaderParamsToCopy, DestinationBounds);
 
-	if (DestinationBounds.IsEmpty())
+	if (DestinationBounds.Area() <= 0)
 	{
 		// Patch must be outside the landscape.
-		return InCombinedResult;
+		return;
 	}
 
-	// Will be used later, once we implement batched merge
-	int32 LandscapeTextureSliceIndex = 0;
-
-	ENQUEUE_RENDER_COMMAND(LandscapeTextureHeightPatch)(
-		[InCombinedResult, LandscapeTextureSliceIndex, ShaderParamsToCopy, Patch, DestinationBounds](FRHICommandListImmediate& RHICmdList)
-		{
-			TRACE_CPUPROFILER_EVENT_SCOPE(LandscapeTextureHeightPatch_Render);
+	ENQUEUE_RENDER_COMMAND(LandscapeTextureWeightPatch)(
+		[RenderThreadLandscapeTextureGetter, LandscapeTextureSliceIndex, ShaderParamsToCopy, Patch, DestinationBounds](FRHICommandListImmediate& RHICmdList)
+	{
+		TRACE_CPUPROFILER_EVENT_SCOPE(LandscapeTextureWeightPatch_Render);
 
-			FRDGBuilder GraphBuilder(RHICmdList, RDG_EVENT_NAME("ApplyTextureHeightPatch"));
+		FRDGBuilder GraphBuilder(RHICmdList, RDG_EVENT_NAME("ApplyTextureWeightPatch"));
 
-			TRefCountPtr<IPooledRenderTarget> DestinationRenderTarget = CreateRenderTarget(
-				InCombinedResult->GetResource()->GetTexture2DRHI(), TEXT("LandscapeTextureWeightPatchOutput"));
-			FRDGTextureRef DestinationTexture = GraphBuilder.RegisterExternalTexture(DestinationRenderTarget);
+		TRefCountPtr<IPooledRenderTarget> DestinationRenderTarget = CreateRenderTarget(
+			RenderThreadLandscapeTextureGetter(), TEXT("LandscapeTextureWeightPatchOutput"));
+		FRDGTextureRef DestinationTexture = GraphBuilder.RegisterExternalTexture(DestinationRenderTarget);
 
-			// Make a copy of the portion of our weightmap input that we're writing to so that we can 
-			// read and write at the same time (needed for blending)
-			FRDGTextureDesc InputCopyDescription = DestinationTexture->Desc;
-			InputCopyDescription.ArraySize = 1;
-			InputCopyDescription.NumMips = 1;
-			InputCopyDescription.Extent = DestinationBounds.Size();
-			FRDGTextureRef InputCopy = GraphBuilder.CreateTexture(InputCopyDescription, TEXT("LandscapeTextureWeightPatchInputCopy"));
+		// Make a copy of the portion of our weightmap input that we're writing to so that we can 
+		// read and write at the same time (needed for blending)
+		FRDGTextureDesc InputCopyDescription = DestinationTexture->Desc;
+		InputCopyDescription.Dimension = ETextureDimension::Texture2D;
+		InputCopyDescription.ArraySize = 1;
+		InputCopyDescription.NumMips = 1;
+		InputCopyDescription.Extent = DestinationBounds.Size();
+		FRDGTextureRef InputCopy = GraphBuilder.CreateTexture(InputCopyDescription, TEXT("LandscapeTextureWeightPatchInputCopy"));
 
-			FRHICopyTextureInfo CopyTextureInfo;
-			CopyTextureInfo.SourceMipIndex = 0;
-			CopyTextureInfo.NumMips = 1;
-			CopyTextureInfo.SourceSliceIndex = LandscapeTextureSliceIndex;
-			CopyTextureInfo.NumSlices = 1;
-			CopyTextureInfo.SourcePosition = FIntVector(DestinationBounds.Min.X, DestinationBounds.Min.Y, 0);
-			CopyTextureInfo.Size = FIntVector(InputCopyDescription.Extent.X, InputCopyDescription.Extent.Y, 0);
-			AddCopyTexturePass(GraphBuilder, DestinationTexture, InputCopy, CopyTextureInfo);
+		FRHICopyTextureInfo CopyTextureInfo;
+		CopyTextureInfo.SourceMipIndex = 0;
+		CopyTextureInfo.NumMips = 1;
+		CopyTextureInfo.SourceSliceIndex = LandscapeTextureSliceIndex;
+		CopyTextureInfo.NumSlices = 1;
+		CopyTextureInfo.SourcePosition = FIntVector(DestinationBounds.Min.X, DestinationBounds.Min.Y, 0);
+		CopyTextureInfo.Size = FIntVector(InputCopyDescription.Extent.X, InputCopyDescription.Extent.Y, 0);
+		AddCopyTexturePass(GraphBuilder, DestinationTexture, InputCopy, CopyTextureInfo);
 
-			FApplyLandscapeTextureWeightPatchPS::FParameters* ShaderParams =
-				GraphBuilder.AllocParameters<FApplyLandscapeTextureWeightPatchPS::FParameters>();
-			*ShaderParams = ShaderParamsToCopy;
+		FApplyLandscapeTextureWeightPatchPS::FParameters* ShaderParams =
+			GraphBuilder.AllocParameters<FApplyLandscapeTextureWeightPatchPS::FParameters>();
+		*ShaderParams = ShaderParamsToCopy;
 
-			TRefCountPtr<IPooledRenderTarget> PatchRenderTarget = CreateRenderTarget(Patch->GetTexture2DRHI(), TEXT("LandscapeTextureWeightPatch"));
-			FRDGTextureRef PatchTexture = GraphBuilder.RegisterExternalTexture(PatchRenderTarget);
-			FRDGTextureSRVRef PatchSRV = GraphBuilder.CreateSRV(FRDGTextureSRVDesc::CreateForMipLevel(PatchTexture, 0));
-			ShaderParams->InWeightPatch = PatchSRV;
-			ShaderParams->InWeightPatchSampler = TStaticSamplerState<SF_Bilinear, AM_Clamp, AM_Clamp>::GetRHI();
+		TRefCountPtr<IPooledRenderTarget> PatchRenderTarget = CreateRenderTarget(Patch->GetTexture2DRHI(), TEXT("LandscapeTextureWeightPatch"));
+		FRDGTextureRef PatchTexture = GraphBuilder.RegisterExternalTexture(PatchRenderTarget);
+		FRDGTextureSRVRef PatchSRV = GraphBuilder.CreateSRV(FRDGTextureSRVDesc::CreateForMipLevel(PatchTexture, 0));
+		ShaderParams->InWeightPatch = PatchSRV;
+		ShaderParams->InWeightPatchSampler = TStaticSamplerState<SF_Bilinear, AM_Clamp, AM_Clamp>::GetRHI();
 
-			FRDGTextureSRVRef InputCopySRV = GraphBuilder.CreateSRV(FRDGTextureSRVDesc::CreateForMipLevel(InputCopy, 0));
-			ShaderParams->InSourceWeightmap = InputCopySRV;
-			ShaderParams->InSourceWeightmapCoordOffset = DestinationBounds.Min;
+		FRDGTextureSRVRef InputCopySRV = GraphBuilder.CreateSRV(FRDGTextureSRVDesc::CreateForMipLevel(InputCopy, 0));
+		ShaderParams->InSourceWeightmap = InputCopySRV;
+		ShaderParams->InSourceWeightmapCoordOffset = DestinationBounds.Min;
 
-			ShaderParams->RenderTargets[0] = FRenderTargetBinding(DestinationTexture, ERenderTargetLoadAction::ENoAction,
-				/*InMipIndex = */0, LandscapeTextureSliceIndex);
+		ShaderParams->RenderTargets[0] = FRenderTargetBinding(DestinationTexture, ERenderTargetLoadAction::ENoAction,
+			/*InMipIndex = */0, LandscapeTextureSliceIndex);
 
-			FApplyLandscapeTextureWeightPatchPS::AddToRenderGraph(GraphBuilder, ShaderParams, DestinationBounds);
+		FApplyLandscapeTextureWeightPatchPS::AddToRenderGraph(GraphBuilder, ShaderParams, DestinationBounds);
 
-			GraphBuilder.Execute();
-		});
+		GraphBuilder.Execute();
+	});
 
-	return InCombinedResult;
 }
 
 void ULandscapeTexturePatch::GetCommonShaderParams(const FTransform& LandscapeHeightmapToWorldIn,
@@ -605,8 +742,6 @@ void ULandscapeTexturePatch::RequestReinitializeWeights()
 	}
 	SetResolution(DesiredResolution);
 
-	const FName VisibilityLayerName = Landscape->VisibilityLayer != nullptr ? Landscape->VisibilityLayer->LayerName : NAME_None;
-
 	ULandscapeInfo* Info = Landscape->GetLandscapeInfo();
 	if (Info)
 	{
@@ -616,22 +751,38 @@ void ULandscapeTexturePatch::RequestReinitializeWeights()
 			{
 				continue;
 			}
+			
 			FName WeightmapLayerName = InfoLayerSettings.GetLayerName();
-			if (!ensure(WeightmapLayerName != NAME_None) 
-				// TODO: We are currently unable to edit the visibility layer, but we get it passed in. Skip it so
-				// that we don't get the users' hopes up.
-				|| WeightmapLayerName == VisibilityLayerName)
+			bool bIsVisibilityLayer = UE::Landscape::IsVisibilityLayer(InfoLayerSettings.LayerInfoObj);
+
+			// Minor note: there's some undefined behavior if a user uses a patch that both has bEditVisibilityLayer
+			//  set to true and a weight layer name that matches some other weight layer. That's ok.
+			TArray<TObjectPtr<ULandscapeWeightPatchTextureInfo>> FoundPatches;
+			if (bIsVisibilityLayer)
 			{
-				continue;
+				FoundPatches = WeightPatches.FilterByPredicate([](const TObjectPtr<ULandscapeWeightPatchTextureInfo>& PatchInfo)
+				{
+					return IsValid(PatchInfo) && PatchInfo->bEditVisibilityLayer;
+				});
+			}
+			else
+			{
+				if (!ensure(WeightmapLayerName != NAME_None))
+				{
+					continue;
+				}
+				FoundPatches = WeightPatches.FilterByPredicate(
+					[&WeightmapLayerName](const TObjectPtr<ULandscapeWeightPatchTextureInfo>& PatchInfo) 
+				{ 
+					return PatchInfo && PatchInfo->WeightmapLayerName == WeightmapLayerName; 
+				});
 			}
-
-			TArray<TObjectPtr<ULandscapeWeightPatchTextureInfo>> FoundPatches = WeightPatches.FilterByPredicate(
-				[&WeightmapLayerName](const TObjectPtr<ULandscapeWeightPatchTextureInfo>& PatchInfo) { return PatchInfo && PatchInfo->WeightmapLayerName == WeightmapLayerName; });
 
 			if (FoundPatches.IsEmpty())
 			{
 				AddWeightPatch(WeightmapLayerName, ELandscapeTexturePatchSourceMode::InternalTexture, false);
 				WeightPatches.Last()->bReinitializeOnNextRender = true;
+				WeightPatches.Last()->bEditVisibilityLayer = bIsVisibilityLayer;
 			}
 			else
 			{
@@ -650,6 +801,8 @@ void ULandscapeTexturePatch::RequestReinitializeWeights()
 #if WITH_EDITOR
 void ULandscapeTexturePatch::ReinitializeHeight(UTextureRenderTarget2D* InCombinedResult, const FTransform& LandscapeHeightmapToWorld)
 {
+	using namespace LandscapeTexturePatchLocals;
+
 	if (HeightSourceMode == ELandscapeTexturePatchSourceMode::TextureAsset)
 	{
 		UE_LOG(LogLandscapePatch, Warning, TEXT("ULandscapeTexturePatch: Cannot reinitialize height patch when source mode is an external texture."));
@@ -714,6 +867,9 @@ void ULandscapeTexturePatch::ReinitializeHeight(UTextureRenderTarget2D* InCombin
 
 	FMatrix44f PatchToSource = GetPatchToHeightmapUVs(LandscapeHeightmapToWorld, TemporaryNativeHeightCopy->SizeX, TemporaryNativeHeightCopy->SizeY, InCombinedResult->SizeX, InCombinedResult->SizeY);
 
+	// TODO: see comment in function
+	DoReinitializationOverlapCheck(PatchToSource, TemporaryNativeHeightCopy->SizeX, TemporaryNativeHeightCopy->SizeY);
+
 	ENQUEUE_RENDER_COMMAND(LandscapeTexturePatchReinitializeHeight)(
 		[Source = InCombinedResult->GetResource(), Destination = TemporaryNativeHeightCopy->GetResource(),
 		&PatchToSource, OffsetToApply](FRHICommandListImmediate& RHICmdList)
@@ -781,11 +937,12 @@ void ULandscapeTexturePatch::ReinitializeHeight(UTextureRenderTarget2D* InCombin
 }
 
 void ULandscapeTexturePatch::ReinitializeWeightPatch(ULandscapeWeightPatchTextureInfo* PatchInfo, 
-	UTextureRenderTarget2D* InCombinedResult, const FTransform& LandscapeHeightmapToWorld)
+	FTextureResource* InputResource, FIntPoint ResourceSize, int32 SliceIndex, 
+	const FTransform& LandscapeHeightmapToWorld)
 {
 	using namespace LandscapeTexturePatchLocals;
 
-	if (!ensure(IsValid(PatchInfo) && IsValid(InCombinedResult)))
+	if (!ensure(IsValid(PatchInfo) && InputResource))
 	{
 		return;
 	}
@@ -834,10 +991,13 @@ void ULandscapeTexturePatch::ReinitializeWeightPatch(ULandscapeWeightPatchTextur
 	}
 
 	FMatrix44f PatchToSource = GetPatchToHeightmapUVs(LandscapeHeightmapToWorld, 
-		RenderTarget->SizeX, RenderTarget->SizeY, InCombinedResult->SizeX, InCombinedResult->SizeY);
+		RenderTarget->SizeX, RenderTarget->SizeY, ResourceSize.X, ResourceSize.Y);
+
+	// TODO: see comment in function
+	DoReinitializationOverlapCheck(PatchToSource, RenderTarget->SizeX, RenderTarget->SizeY);
 
 	ENQUEUE_RENDER_COMMAND(LandscapeTexturePatchReinitializeWeight)(
-		[Source = InCombinedResult->GetResource(), Destination = RenderTarget->GetResource(), &PatchToSource](FRHICommandListImmediate& RHICmdList)
+		[InputResource, SliceIndex, Destination = RenderTarget->GetResource(), &PatchToSource](FRHICommandListImmediate& RHICmdList)
 	{
 		using namespace UE::Landscape;
 
@@ -845,9 +1005,22 @@ void ULandscapeTexturePatch::ReinitializeWeightPatch(ULandscapeWeightPatchTextur
 
 		FReinitializeLandscapePatchPS::FParameters* ShaderParams = GraphBuilder.AllocParameters<FReinitializeLandscapePatchPS::FParameters>();
 
-		FRDGTextureRef SourceTexture = GraphBuilder.RegisterExternalTexture(CreateRenderTarget(Source->GetTexture2DRHI(), TEXT("ReinitializationSource")));
-		FRDGTextureSRVRef SourceSRV = GraphBuilder.CreateSRV(FRDGTextureSRVDesc::CreateForMipLevel(SourceTexture, 0));
-		ShaderParams->InSource = SourceSRV;
+		if (SliceIndex < 0)
+		{
+			FRDGTextureRef SourceTexture = GraphBuilder.RegisterExternalTexture(CreateRenderTarget(
+				InputResource->GetTexture2DRHI(), TEXT("ReinitializationSource")));
+			ShaderParams->InSource = GraphBuilder.CreateSRV(FRDGTextureSRVDesc::CreateForMipLevel(SourceTexture, 0));
+		
+		}
+		else
+		{
+			FRDGTextureRef SourceTexture = GraphBuilder.RegisterExternalTexture(CreateRenderTarget(
+				InputResource->GetTexture2DArrayRHI(), TEXT("ReinitializationSource")));
+			FRDGTextureSRVDesc Desc = FRDGTextureSRVDesc::CreateForSlice(SourceTexture, SliceIndex);
+			Desc.MipLevel = 0;
+			Desc.NumMipLevels = 1;
+			ShaderParams->InSource = GraphBuilder.CreateSRV(Desc);
+		}
 
 		ShaderParams->InSourceSampler = TStaticSamplerState<SF_Bilinear, AM_Clamp, AM_Clamp>::GetRHI();
 
@@ -885,35 +1058,48 @@ FMatrix44f ULandscapeTexturePatch::GetPatchToHeightmapUVs(const FTransform& Land
 
 bool ULandscapeTexturePatch::CanAffectHeightmap() const
 {
-	return HeightSourceMode != ELandscapeTexturePatchSourceMode::None || bReinitializeHeightOnNextRender;
+	return (HeightSourceMode != ELandscapeTexturePatchSourceMode::None
+		// If source mode is texture asset, we need to have an asset to read from
+		&& (HeightSourceMode != ELandscapeTexturePatchSourceMode::TextureAsset || HeightTextureAsset))
+		// If reinitializing, we need to read from the render call
+		|| bReinitializeHeightOnNextRender;
 }
 
 bool ULandscapeTexturePatch::CanAffectWeightmap() const
 {
-	return Algo::AnyOf(WeightPatches, [](const TObjectPtr<ULandscapeWeightPatchTextureInfo>& InWeightPatch) 
+	return Algo::AnyOf(WeightPatches, [this](const TObjectPtr<ULandscapeWeightPatchTextureInfo>& InWeightPatch) 
 	{ 
-		return IsValid(InWeightPatch) && (InWeightPatch->SourceMode != ELandscapeTexturePatchSourceMode::None || InWeightPatch->bReinitializeOnNextRender);
+		return IsValid(InWeightPatch) && WeightPatchCanRender(*InWeightPatch);
 	});
 }
 
 bool ULandscapeTexturePatch::CanAffectWeightmapLayer(const FName& InLayerName) const
 {
-	return Algo::AnyOf(WeightPatches, [InLayerName](TObjectPtr<ULandscapeWeightPatchTextureInfo> InWeightPatch) 
+	return Algo::AnyOf(WeightPatches, [InLayerName, this](TObjectPtr<ULandscapeWeightPatchTextureInfo> InWeightPatch) 
 	{
-		return IsValid(InWeightPatch) && (InWeightPatch->SourceMode != ELandscapeTexturePatchSourceMode::None || InWeightPatch->bReinitializeOnNextRender)
-			&& (InWeightPatch->WeightmapLayerName == InLayerName);
+		return IsValid(InWeightPatch) && (InWeightPatch->WeightmapLayerName == InLayerName)
+			&& WeightPatchCanRender(*InWeightPatch);
 	}); 
 }
 
 bool ULandscapeTexturePatch::CanAffectVisibilityLayer() const
 {
-	return Algo::AnyOf(WeightPatches, [](const TObjectPtr<ULandscapeWeightPatchTextureInfo>& InWeightPatch) 
+	return Algo::AnyOf(WeightPatches, [this](const TObjectPtr<ULandscapeWeightPatchTextureInfo>& InWeightPatch) 
 	{ 
 		return IsValid(InWeightPatch) && InWeightPatch->bEditVisibilityLayer 
-			&& (InWeightPatch->SourceMode != ELandscapeTexturePatchSourceMode::None || InWeightPatch->bReinitializeOnNextRender);
+			&& WeightPatchCanRender(*InWeightPatch);
 	});
 }
 
+bool ULandscapeTexturePatch::WeightPatchCanRender(const ULandscapeWeightPatchTextureInfo& InWeightPatch) const
+{
+	return (InWeightPatch.SourceMode != ELandscapeTexturePatchSourceMode::None 
+		// If source mode is texture asset, we need to have an asset to read from
+		&& (InWeightPatch.SourceMode != ELandscapeTexturePatchSourceMode::TextureAsset || InWeightPatch.TextureAsset))
+		// If reinitializing, we need to read from the render call
+		|| InWeightPatch.bReinitializeOnNextRender;
+}
+
 void ULandscapeTexturePatch::GetRenderDependencies(TSet<UObject*>& OutDependencies) const
 {
 	Super::GetRenderDependencies(OutDependencies);
@@ -1710,3 +1896,75 @@ void ULandscapeTexturePatch::SetEditVisibilityLayer(const FName& InWeightmapLaye
 		}
 	}
 }
+
+#if WITH_EDITOR
+void ULandscapeTexturePatch::GetRendererStateInfo(const ULandscapeInfo* InLandscapeInfo,
+	FEditLayerTargetTypeState& OutSupported, FEditLayerTargetTypeState& OutEnabled,
+	TArray<TSet<FName>>& OutRenderGroups) const
+{
+	if (CanAffectHeightmap())
+	{
+		OutSupported.AddTargetType(ELandscapeToolTargetType::Heightmap);
+	}
+
+	for (const TObjectPtr<ULandscapeWeightPatchTextureInfo>& WeightPatch : WeightPatches)
+	{
+		if (IsValid(WeightPatch) && WeightPatchCanRender(*WeightPatch))
+		{
+			if (WeightPatch->bEditVisibilityLayer)
+			{
+				OutSupported.AddTargetType(ELandscapeToolTargetType::Visibility);
+			}
+			else
+			{
+				OutSupported.AddTargetType(ELandscapeToolTargetType::Weightmap);
+				OutSupported.AddWeightmap(WeightPatch->WeightmapLayerName);
+				OutRenderGroups.Add(TSet<FName>({ WeightPatch->WeightmapLayerName }));
+			}
+		}
+	}
+
+	if (IsEnabled())
+	{
+		OutEnabled = OutSupported;
+	}
+}
+
+FString ULandscapeTexturePatch::GetEditLayerRendererDebugName() const
+{
+	return FString::Printf(TEXT("ULandscapeTexturePatch:%s"), *GetFullName());
+}
+
+TArray<UE::Landscape::EditLayers::FEditLayerRenderItem> ULandscapeTexturePatch::GetRenderItems(const ULandscapeInfo* InLandscapeInfo) const
+{
+	using namespace UE::Landscape::EditLayers;
+
+	TArray<FEditLayerRenderItem> AffectedAreas;
+
+	FTransform ComponentTransform = this->GetComponentToWorld();
+	FOOBox2D PatchArea(ComponentTransform, GetFullUnscaledWorldSize());
+	FInputWorldArea InputWorldArea = FInputWorldArea::CreateOOBox(PatchArea);
+	FOutputWorldArea OutputWorldArea = FOutputWorldArea::CreateOOBox(PatchArea);
+
+	if (CanAffectHeightmap())
+	{
+		FEditLayerTargetTypeState TargetInfo(ELandscapeToolTargetTypeFlags::Heightmap);
+		FEditLayerRenderItem Item(TargetInfo, InputWorldArea, OutputWorldArea, false);
+		AffectedAreas.Add(Item);
+	}
+
+	for (const TObjectPtr<ULandscapeWeightPatchTextureInfo>& WeightPatch : WeightPatches)
+	{
+		if (IsValid(WeightPatch) && WeightPatchCanRender(*WeightPatch))
+		{
+			FEditLayerTargetTypeState TargetInfo = WeightPatch->bEditVisibilityLayer ? 
+				FEditLayerTargetTypeState(ELandscapeToolTargetTypeFlags::Visibility)
+				: FEditLayerTargetTypeState(ELandscapeToolTargetTypeFlags::Weightmap, { WeightPatch->WeightmapLayerName });
+			FEditLayerRenderItem Item(TargetInfo, InputWorldArea, OutputWorldArea, false);
+			AffectedAreas.Add(Item);
+		}
+	}
+
+	return AffectedAreas;
+}
+#endif // WITH_EDITOR
diff --git a/Engine/Plugins/Experimental/LandscapePatch/Source/LandscapePatch/Public/LandscapeCircleHeightPatch.h b/Engine/Plugins/Experimental/LandscapePatch/Source/LandscapePatch/Public/LandscapeCircleHeightPatch.h
--- a/Engine/Plugins/Experimental/LandscapePatch/Source/LandscapePatch/Public/LandscapeCircleHeightPatch.h
+++ b/Engine/Plugins/Experimental/LandscapePatch/Source/LandscapePatch/Public/LandscapeCircleHeightPatch.h
@@ -21,6 +21,20 @@ public:
 
 	virtual UTextureRenderTarget2D* RenderLayer_Native(const FLandscapeBrushParameters& InParameters, const FTransform& HeightmapToWorld) override;
 
+#if WITH_EDITOR
+	using FEditLayerTargetTypeState = UE::Landscape::EditLayers::FEditLayerTargetTypeState;
+	using FEditLayerRenderItem = UE::Landscape::EditLayers::FEditLayerRenderItem;
+
+	// ILandscapeEditLayerRenderer, via ULandscapePatchComponent
+	virtual void GetRendererStateInfo(const ULandscapeInfo* InLandscapeInfo,
+		FEditLayerTargetTypeState& OutSupportedTargetTypeState, 
+		FEditLayerTargetTypeState& OutEnabledTargetTypeState, 
+		TArray<TSet<FName>>& OutRenderGroups) const override;
+	virtual FString GetEditLayerRendererDebugName() const override;
+	virtual TArray<FEditLayerRenderItem> GetRenderItems(const ULandscapeInfo* InLandscapeInfo) const override;
+	virtual void RenderLayer(FRenderParams& InRenderParams) override;
+#endif
+
 	virtual bool CanAffectHeightmap() const override { return !bEditVisibility; }
 	virtual bool CanAffectVisibilityLayer() const override { return bEditVisibility; }
 
@@ -36,7 +50,7 @@ protected:
 	UPROPERTY(EditAnywhere, Category = Settings)
 	float Falloff = 500;
 
-	/** Specifies if this patch edits the visibility layer. */
+	/** Specifies if this patch edits the visibility layer instead of height. */
 	UPROPERTY(EditAnywhere, Category = Settings)
 	bool bEditVisibility = false;
 
@@ -44,6 +58,11 @@ protected:
 	  vertices are also included and the whole circle is able to lie flat. */
 	UPROPERTY(EditAnywhere, Category = Settings, AdvancedDisplay)
 	bool bExclusiveRadius = false;
+
+private:
+	void ApplyCirclePatch(bool bIsVisibilityLayer,
+		const TFunction<FRHITexture* ()>& RenderThreadLandscapeTextureGetter, int32 LandscapeTextureSliceIndex,
+		const FIntPoint& DestinationResolution, const FTransform& HeightmapCoordsToWorld);
 };
 
 #if UE_ENABLE_INCLUDE_ORDER_DEPRECATED_IN_5_2
diff --git a/Engine/Plugins/Experimental/LandscapePatch/Source/LandscapePatch/Public/LandscapeCircleHeightPatchPS.h b/Engine/Plugins/Experimental/LandscapePatch/Source/LandscapePatch/Public/LandscapeCircleHeightPatchPS.h
--- a/Engine/Plugins/Experimental/LandscapePatch/Source/LandscapePatch/Public/LandscapeCircleHeightPatchPS.h
+++ b/Engine/Plugins/Experimental/LandscapePatch/Source/LandscapePatch/Public/LandscapeCircleHeightPatchPS.h
@@ -16,6 +16,11 @@ class FLandscapeCircleHeightPatchPS : public FGlobalShader
 public:
 	BEGIN_SHADER_PARAMETER_STRUCT(FParameters, )
 		SHADER_PARAMETER_RDG_TEXTURE_SRV(Texture2D<float4>, InSourceTexture) // Our input texture
+		// Offset of the source heightmap relative to the 0,0 location in the destination heightmap, because
+		// the source is likely to be a copied region from some inner part of the destination. This is basically
+		// a SourceHeightmapToDestinationHeightmap coordinate transformation, except that it is always a simple 
+		// integer translation.
+		SHADER_PARAMETER(FIntPoint, InSourceTextureOffset)
 		SHADER_PARAMETER(FVector3f, InCenter)
 		SHADER_PARAMETER(float, InRadius)
 		SHADER_PARAMETER(float, InFalloff)
@@ -25,7 +30,7 @@ public:
 	static bool ShouldCompilePermutation(const FGlobalShaderPermutationParameters& Parameters);
 	static void ModifyCompilationEnvironment(const FGlobalShaderPermutationParameters& Parameters, FShaderCompilerEnvironment& OutEnvironment);
 
-	LANDSCAPEPATCH_API static void AddToRenderGraph(FRDGBuilder& GraphBuilder, FParameters* InParameters);
+	LANDSCAPEPATCH_API static void AddToRenderGraph(FRDGBuilder& GraphBuilder, FParameters* InParameters, const FIntRect& DestinationBounds);
 };
 
 /**
@@ -42,7 +47,7 @@ public:
 	{}
 
 	static void ModifyCompilationEnvironment(const FGlobalShaderPermutationParameters& Parameters, FShaderCompilerEnvironment& OutEnvironment);
-	LANDSCAPEPATCH_API static void AddToRenderGraph(FRDGBuilder& GraphBuilder, FParameters* InParameters);
+	LANDSCAPEPATCH_API static void AddToRenderGraph(FRDGBuilder& GraphBuilder, FParameters* InParameters, const FIntRect& DestinationBounds);
 };
 
 #if UE_ENABLE_INCLUDE_ORDER_DEPRECATED_IN_5_2
diff --git a/Engine/Plugins/Experimental/LandscapePatch/Source/LandscapePatch/Public/LandscapePatchComponent.h b/Engine/Plugins/Experimental/LandscapePatch/Source/LandscapePatch/Public/LandscapePatchComponent.h
--- a/Engine/Plugins/Experimental/LandscapePatch/Source/LandscapePatch/Public/LandscapePatchComponent.h
+++ b/Engine/Plugins/Experimental/LandscapePatch/Source/LandscapePatch/Public/LandscapePatchComponent.h
@@ -5,6 +5,7 @@
 
 #include "Components/SceneComponent.h"
 #include "LandscapeBlueprintBrushBase.h"
+#include "LandscapeEditLayerRenderer.h"
 #include "LandscapeEditTypes.h"
 #include "LandscapePatchEditLayer.h" // PATCH_PRIORITY_BASE
 #include "UObject/WeakInterfacePtr.h"
@@ -45,6 +46,7 @@ enum class ELandscapePatchPriorityInitialization : uint8
 //~ so that we can use render proxies for passing along data to the render thread or perhaps for visualization.
 UCLASS(Blueprintable, BlueprintType, Abstract)
 class LANDSCAPEPATCH_API ULandscapePatchComponent : public USceneComponent
+	, public ILandscapeEditLayerRenderer
 {
 	GENERATED_BODY()
 
@@ -55,6 +57,12 @@ public:
 	virtual UTextureRenderTarget2D* RenderLayer_Native(const FLandscapeBrushParameters& InParameters, 
 		const FTransform& HeightmapToWorld) { return InParameters.CombinedResult; }
 
+#if WITH_EDITOR
+	// ILandscapeEditLayerRenderer
+	virtual bool CanRender() const override { return CanAffectLandscape(); }
+	// Subclasses are expected to implement GetEditLayerRendererDebugName, GetRendererStateInfo, GetRenderItems, and RenderLayer
+#endif
+
 	// These determine whether the patch is configured correctly to affect height/weightmaps,
 	// ignoring whether it is currently enabled or not.
 	virtual bool CanAffectHeightmap() const { return false; }
diff --git a/Engine/Plugins/Experimental/LandscapePatch/Source/LandscapePatch/Public/LandscapePatchEditLayer.h b/Engine/Plugins/Experimental/LandscapePatch/Source/LandscapePatch/Public/LandscapePatchEditLayer.h
--- a/Engine/Plugins/Experimental/LandscapePatch/Source/LandscapePatch/Public/LandscapePatchEditLayer.h
+++ b/Engine/Plugins/Experimental/LandscapePatch/Source/LandscapePatch/Public/LandscapePatchEditLayer.h
@@ -21,6 +21,7 @@ class ULandscapePatchEditLayer : public ULandscapeEditLayerProcedural
 	GENERATED_BODY()
 
 public:
+
 #if WITH_EDITOR
 	/**
 	 * Must be called by patches on this layer to let the layer know that the patch is pointing
@@ -60,14 +61,11 @@ public:
 	virtual void OnLayerCreated(FLandscapeLayer& Layer);
 
 #if WITH_EDITOR
-	// TODO [jonathan.bard] : The following implementation is completely bogus, for now, until this supports batched merge : 
-	//~ Begin ILandscapeEditLayerRenderer implementation
-	virtual void GetRendererStateInfo(const ULandscapeInfo* InLandscapeInfo,
-		UE::Landscape::EditLayers::FEditLayerTargetTypeState& OutSupportedTargetTypeState, UE::Landscape::EditLayers::FEditLayerTargetTypeState& OutEnabledTargetTypeState, TArray<TSet<FName>>& OutRenderGroups) const override {}
-	virtual TArray<UE::Landscape::EditLayers::FEditLayerRenderItem> GetRenderItems(const ULandscapeInfo* InLandscapeInfo) const override { return {}; }
-	virtual FString GetEditLayerRendererDebugName() const override { return TEXT("LandscapePatchEditLayer"); }
-	//~ End ILandscapeEditLayerRenderer implementation
-	
+	using FEditLayerRendererState = UE::Landscape::EditLayers::FEditLayerRendererState;
+
+	// IEditLayerRendererProvider
+	virtual TArray<FEditLayerRendererState> GetEditLayerRendererStates(const ULandscapeInfo* InLandscapeInfo, bool bInSkipBrush) override;
+
 	// UObject
 	virtual void PostLoad() override;
 	virtual void PostEditUndo() override;
diff --git a/Engine/Plugins/Experimental/LandscapePatch/Source/LandscapePatch/Public/LandscapePatchManager.h b/Engine/Plugins/Experimental/LandscapePatch/Source/LandscapePatch/Public/LandscapePatchManager.h
--- a/Engine/Plugins/Experimental/LandscapePatch/Source/LandscapePatch/Public/LandscapePatchManager.h
+++ b/Engine/Plugins/Experimental/LandscapePatch/Source/LandscapePatch/Public/LandscapePatchManager.h
@@ -33,9 +33,25 @@ public:
 	virtual UTextureRenderTarget2D* RenderLayer_Native(const FLandscapeBrushParameters& InParameters) override;
 
 #if WITH_EDITOR
-	//~ Begin ILandscapeEditLayerRenderer implementation
-	virtual TArray<UE::Landscape::EditLayers::FEditLayerRenderItem> GetRenderItems(const ULandscapeInfo* InLandscapeInfo) const override;
-	//~ End ILandscapeEditLayerRenderer implementation
+	using FEditLayerRendererState = UE::Landscape::EditLayers::FEditLayerRendererState;
+
+	// IEditLayerRendererProvider
+	// Called by the batched-merge application path.
+	virtual TArray<FEditLayerRendererState> GetEditLayerRendererStates(const ULandscapeInfo* InLandscapeInfo, bool bInSkipBrush) override;
+
+	// ILandscapeEditLayerRenderer
+	//~ In batched merge, the manager relies on being a renderer provider. It does not need to have its 
+	//~  RenderLayer_Native method called, so we override the implementations of ILandscapeEditLayerRenderer
+	//~  inherited from ALandscapeBlueprintBrushBase to do nothing.
+	virtual void GetRendererStateInfo(const ULandscapeInfo* InLandscapeInfo,
+		UE::Landscape::EditLayers::FEditLayerTargetTypeState& OutSupportedTargetTypeState,
+		UE::Landscape::EditLayers::FEditLayerTargetTypeState& OutEnabledTargetTypeState,
+		TArray<TSet<FName>>& OutRenderGroups) const override {}
+	virtual TArray<UE::Landscape::EditLayers::FEditLayerRenderItem> GetRenderItems(const ULandscapeInfo* InLandscapeInfo) const override
+	{
+		return {};
+	}
+	virtual void RenderLayer(ILandscapeEditLayerRenderer::FRenderParams& InRenderParams) override { };
 #endif // WITH_EDITOR
 
 	// Adds the brush to the given landscape, removing it from any previous one. This differs from SetOwningLandscape
diff --git a/Engine/Plugins/Experimental/LandscapePatch/Source/LandscapePatch/Public/LandscapeTexturePatch.h b/Engine/Plugins/Experimental/LandscapePatch/Source/LandscapePatch/Public/LandscapeTexturePatch.h
--- a/Engine/Plugins/Experimental/LandscapePatch/Source/LandscapePatch/Public/LandscapeTexturePatch.h
+++ b/Engine/Plugins/Experimental/LandscapePatch/Source/LandscapePatch/Public/LandscapeTexturePatch.h
@@ -220,6 +220,17 @@ class LANDSCAPEPATCH_API ULandscapeTexturePatch : public ULandscapePatchComponen
 public:
 
 #if WITH_EDITOR
+	using FEditLayerTargetTypeState = UE::Landscape::EditLayers::FEditLayerTargetTypeState;
+	using FEditLayerRenderItem = UE::Landscape::EditLayers::FEditLayerRenderItem;
+
+	// ILandscapeEditLayerRenderer, via ULandscapePatchComponent
+	virtual void GetRendererStateInfo(const ULandscapeInfo* InLandscapeInfo,
+		FEditLayerTargetTypeState& OutSupportedTargetTypeState,
+		FEditLayerTargetTypeState& OutEnabledTargetTypeState,
+		TArray<TSet<FName>>& OutRenderGroups) const override;
+	virtual FString GetEditLayerRendererDebugName() const override;
+	virtual TArray<FEditLayerRenderItem> GetRenderItems(const ULandscapeInfo* InLandscapeInfo) const override;
+	virtual void RenderLayer(FRenderParams& InRenderParams) override;
 
 	// ULandscapePatchComponent
 	UTextureRenderTarget2D* RenderLayer_Native(const FLandscapeBrushParameters& InParameters, const FTransform& HeightmapToWorld) override;
@@ -634,8 +645,9 @@ private:
 	void TransitionHeightSourceModeInternal(ELandscapeTexturePatchSourceMode OldMode, ELandscapeTexturePatchSourceMode NewMode);
 	FLandscapeHeightPatchConvertToNativeParams GetHeightConvertToNativeParams() const;
 	UTextureRenderTarget2D* ApplyToHeightmap(UTextureRenderTarget2D* InCombinedResult, const FTransform& LandscapeHeightmapToWorld);
-	UTextureRenderTarget2D* ApplyToWeightmap(ULandscapeWeightPatchTextureInfo* PatchInfo, 
-		UTextureRenderTarget2D* InCombinedResult, const FTransform& LandscapeHeightmapToWorld);
+	void ApplyToWeightmap(ULandscapeWeightPatchTextureInfo* PatchInfo,
+		TFunction<FRHITexture* ()> RenderThreadLandscapeTextureGetter, int32 LandscapeTextureSliceIndex,
+		const FIntPoint& LandscapeTextureResolution, const FTransform& LandscapeHeightmapToWorld);
 
 	void GetCommonShaderParams(const FTransform& LandscapeHeightmapToWorldIn, 
 		const FIntPoint& SourceResolutionIn, const FIntPoint& DestinationResolutionIn,
@@ -649,9 +661,15 @@ private:
 		UE::Landscape::FApplyLandscapeTextureWeightPatchPS::FParameters& ParamsOut, FIntRect& DestinationBoundsOut) const;
 	FMatrix44f GetPatchToHeightmapUVs(const FTransform& LandscapeHeightmapToWorld, int32 PatchSizeX, int32 PatchSizeY, int32 HeightmapSizeX, int32 HeightmapSizeY) const;
 	void ReinitializeHeight(UTextureRenderTarget2D* InCombinedResult, const FTransform& LandscapeHeightmapToWorld);
-	void ReinitializeWeightPatch(ULandscapeWeightPatchTextureInfo* PatchInfo, UTextureRenderTarget2D* InCombinedResult, 
+	/**
+	 * @param SliceIndex set to a negative value when not using a Texture2DArray
+	 */
+	void ReinitializeWeightPatch(ULandscapeWeightPatchTextureInfo* PatchInfo,
+		FTextureResource* InputResource, FIntPoint ResourceSize, int32 SliceIndex,
 		const FTransform& LandscapeHeightmapToWorld);
 
+	bool WeightPatchCanRender(const ULandscapeWeightPatchTextureInfo& InWeightPatch) const;
+
 	void ResetHeightRenderTargetFormat();
 #endif // WITH_EDITOR
 
diff --git a/Engine/Source/Runtime/Landscape/Classes/LandscapeEditLayer.h b/Engine/Source/Runtime/Landscape/Classes/LandscapeEditLayer.h
--- a/Engine/Source/Runtime/Landscape/Classes/LandscapeEditLayer.h
+++ b/Engine/Source/Runtime/Landscape/Classes/LandscapeEditLayer.h
@@ -17,7 +17,7 @@ class ALandscape;
 UCLASS(MinimalAPI, Abstract)
 class ULandscapeEditLayerBase : public UObject
 #if CPP && WITH_EDITOR
-	, public ILandscapeEditLayerRenderer
+	, public UE::Landscape::EditLayers::IEditLayerRendererProvider
 #endif // CPP && WITH_EDITOR
 {
 	GENERATED_BODY()
@@ -152,6 +152,18 @@ public:
 	// TODO: Should this be protected and then we friend ALandscape?
 	virtual void OnLayerRemoved() {}
 
+#if WITH_EDITOR
+	//~ Begin IEditLayerRendererProvider implementation
+	// By default this does nothing in a landscape edit layer, but subclasses can override it if 
+	//  they would like to provide additional renderers.
+	virtual TArray<UE::Landscape::EditLayers::FEditLayerRendererState> GetEditLayerRendererStates(
+		const ULandscapeInfo* InLandscapeInfo, bool bInSkipBrush) 
+	{ 
+		return {};
+	};
+	//~ End IEditLayerRendererProvider implementation
+#endif
+
 	// UObject
 	LANDSCAPE_API virtual void PostLoad() override;
 
@@ -172,6 +184,7 @@ protected:
 */
 UCLASS(MinimalAPI, Abstract)
 class ULandscapeEditLayerPersistent : public ULandscapeEditLayerBase
+	, public ILandscapeEditLayerRenderer
 {
 	GENERATED_BODY()
 
diff --git a/Engine/Source/Runtime/Landscape/Private/LandscapeEditLayers.cpp b/Engine/Source/Runtime/Landscape/Private/LandscapeEditLayers.cpp
--- a/Engine/Source/Runtime/Landscape/Private/LandscapeEditLayers.cpp
+++ b/Engine/Source/Runtime/Landscape/Private/LandscapeEditLayers.cpp
@@ -5382,9 +5382,15 @@ TArray<UE::Landscape::EditLayers::FEditLayerRendererState> FLandscapeLayer::GetE
 	TArray<FEditLayerRendererState> RendererStates;
 	RendererStates.Reserve(1 + Brushes.Num());
 
-	// The edit layer itself is a renderer : 
 	check(EditLayer != nullptr);
-	RendererStates.Emplace(EditLayer, InLandscapeInfo);
+	// The edit layer itself might be a renderer:
+	if (ILandscapeEditLayerRenderer* AsRenderer = Cast<ILandscapeEditLayerRenderer>(EditLayer))
+	{
+		RendererStates.Emplace(AsRenderer, InLandscapeInfo);
+	}
+
+	// The layer can also be a renderer provider
+	RendererStates.Append(EditLayer->GetEditLayerRendererStates(InLandscapeInfo, bInSkipBrush));
 
 	for (FLandscapeLayerBrush& Brush : Brushes)
 	{
@@ -9793,13 +9799,18 @@ int32 ALandscape::PerformLayersWeightmapsGlobalMerge(FUpdateLayersContentContext
 							FLandscapeLayerBrush& Brush = Layer.Brushes[i];
 							TOptional<int32> LayerInfoSettingsAllocatedIndex;
 
-							if (Brush.AffectsWeightmapLayer(InfoLayerSettings.GetLayerName()) && !LayerInfoObjects.Contains(InfoLayerSettings.LayerInfoObj))
+							// Note that we need to check visibility before we check by name, or else a name match could
+							//  cause us to allocate the visibility layer at the wrong index.
+							if (UE::Landscape::IsVisibilityLayer(InfoLayerSettings.LayerInfoObj))
 							{
-								LayerInfoSettingsAllocatedIndex = LayerInfoSettingsIndex + 1; // due to visibility layer that is at 0
+								if (Brush.AffectsVisibilityLayer() && !LayerInfoObjects.Contains(InfoLayerSettings.LayerInfoObj))
+								{
+									LayerInfoSettingsAllocatedIndex = GetVisibilityLayerAllocationIndex();
+								}
 							}
-							else if (Brush.AffectsVisibilityLayer() && UE::Landscape::IsVisibilityLayer(InfoLayerSettings.LayerInfoObj) && !LayerInfoObjects.Contains(InfoLayerSettings.LayerInfoObj))
+							else if (Brush.AffectsWeightmapLayer(InfoLayerSettings.GetLayerName()) && !LayerInfoObjects.Contains(InfoLayerSettings.LayerInfoObj))
 							{
-								LayerInfoSettingsAllocatedIndex = GetVisibilityLayerAllocationIndex();
+								LayerInfoSettingsAllocatedIndex = LayerInfoSettingsIndex + 1; // due to visibility layer that is at 0
 							}
 
 							if (LayerInfoSettingsAllocatedIndex.IsSet())
@@ -9814,13 +9825,16 @@ int32 ALandscape::PerformLayersWeightmapsGlobalMerge(FUpdateLayersContentContext
 						if (RenderCallAdapter && !LayerInfoObjects.Contains(InfoLayerSettings.LayerInfoObj))
 						{
 							TOptional<int32> LayerInfoSettingsAllocatedIndex;
-							if (RenderCallAdapter->AffectsWeightmapLayerAsBlueprintBrush(InfoLayerSettings.GetLayerName()))
+							if (UE::Landscape::IsVisibilityLayer(InfoLayerSettings.LayerInfoObj))
 							{
-								LayerInfoSettingsAllocatedIndex = LayerInfoSettingsIndex + 1; // due to visibility layer that is at 0
+								if (RenderCallAdapter->AffectsVisibilityLayerAsBlueprintBrush())
+								{
+									LayerInfoSettingsAllocatedIndex = GetVisibilityLayerAllocationIndex();
+								}
 							}
-							else if (RenderCallAdapter->AffectsVisibilityLayerAsBlueprintBrush() && UE::Landscape::IsVisibilityLayer(InfoLayerSettings.LayerInfoObj))
+							else if (RenderCallAdapter->AffectsWeightmapLayerAsBlueprintBrush(InfoLayerSettings.GetLayerName()))
 							{
-								LayerInfoSettingsAllocatedIndex = GetVisibilityLayerAllocationIndex();
+								LayerInfoSettingsAllocatedIndex = LayerInfoSettingsIndex + 1; // due to visibility layer that is at 0
 							}
 
 							if (LayerInfoSettingsAllocatedIndex.IsSet())
diff --git a/Engine/Source/Runtime/Landscape/Public/LandscapeEditLayerRenderer.h b/Engine/Source/Runtime/Landscape/Public/LandscapeEditLayerRenderer.h
--- a/Engine/Source/Runtime/Landscape/Public/LandscapeEditLayerRenderer.h
+++ b/Engine/Source/Runtime/Landscape/Public/LandscapeEditLayerRenderer.h
@@ -456,10 +456,10 @@ public:
 	inline FName GetTargetLayerNameForIndexChecked(int32 InIndex) const { check(AllTargetLayerNames.IsValidIndex(InIndex)); return AllTargetLayerNames[InIndex]; }
 	TBitArray<> ConvertTargetLayerNamesToBitIndices(TConstArrayView<FName> InTargetLayerNames) const;
 	TBitArray<> ConvertTargetLayerNamesToBitIndicesChecked(TConstArrayView<FName> InTargetLayerNames) const;
-	TArray<FName> ConvertTargetLayerBitIndicesToNames(const TBitArray<>& InTargetLayerBitIndices) const;
-	TArray<FName> ConvertTargetLayerBitIndicesToNamesChecked(const TBitArray<>& InTargetLayerBitIndices) const;
-	TArray<ULandscapeLayerInfoObject*> ConvertTargetLayerBitIndicesToLayerInfos(const TBitArray<>& InTargetLayerBitIndices) const;
-	TArray<ULandscapeLayerInfoObject*> ConvertTargetLayerBitIndicesToLayerInfosChecked(const TBitArray<>& InTargetLayerBitIndices) const;
+	LANDSCAPE_API TArray<FName> ConvertTargetLayerBitIndicesToNames(const TBitArray<>& InTargetLayerBitIndices) const;
+	LANDSCAPE_API TArray<FName> ConvertTargetLayerBitIndicesToNamesChecked(const TBitArray<>& InTargetLayerBitIndices) const;
+	LANDSCAPE_API TArray<ULandscapeLayerInfoObject*> ConvertTargetLayerBitIndicesToLayerInfos(const TBitArray<>& InTargetLayerBitIndices) const;
+	LANDSCAPE_API TArray<ULandscapeLayerInfoObject*> ConvertTargetLayerBitIndicesToLayerInfosChecked(const TBitArray<>& InTargetLayerBitIndices) const;
 
 	/**
 	 * Runs the given function for each all valid target layer in the bit indices in parameters, with the possibility of early exit
diff --git a/Engine/Source/Runtime/Landscape/Public/LandscapeEditResourcesSubsystem.h b/Engine/Source/Runtime/Landscape/Public/LandscapeEditResourcesSubsystem.h
--- a/Engine/Source/Runtime/Landscape/Public/LandscapeEditResourcesSubsystem.h
+++ b/Engine/Source/Runtime/Landscape/Public/LandscapeEditResourcesSubsystem.h
@@ -81,28 +81,28 @@ public:
 	ULandscapeScratchRenderTarget();
 
 	/** @return a debug name while this render target is in use. This helps track things down, since the underlying resource can be recycled */
-	const FString& GetDebugName() const;
+	LANDSCAPE_API const FString& GetDebugName() const;
 	/** @return Getter for the internal UTextureRenderTarget */
 	UTextureRenderTarget* GetRenderTarget() const { return RenderTarget; }
 	/** @return Getter for the internal UTextureRenderTarget2D (only when CurrentRenderTargetParams.NumSlices == 0, asserts if CurrentRenderTargetParams.NumSlices > 0) */
-	UTextureRenderTarget2D* GetRenderTarget2D() const;
+	LANDSCAPE_API UTextureRenderTarget2D* GetRenderTarget2D() const;
 	/** @return the internal UTextureRenderTarget2D (only when CurrentRenderTargetParams.NumSlices == 0, returns nullptr if CurrentRenderTargetParams.NumSlices > 0) */
-	UTextureRenderTarget2D* TryGetRenderTarget2D() const;
+	LANDSCAPE_API UTextureRenderTarget2D* TryGetRenderTarget2D() const;
 	/**  @return the internal UTextureRenderTarget2DArray (only when CurrentRenderTargetParams.NumSlices > 0, asserts if CurrentRenderTargetParams.NumSlices == 0)  */
-	UTextureRenderTarget2DArray* GetRenderTarget2DArray() const;
+	LANDSCAPE_API UTextureRenderTarget2DArray* GetRenderTarget2DArray() const;
 	/** @return the internal UTextureRenderTarget2DArray (only when CurrentRenderTargetParams.NumSlices > 0, returns nullptr if CurrentRenderTargetParams.NumSlices == 0) */
-	UTextureRenderTarget2DArray* TryGetRenderTarget2DArray() const;
+	LANDSCAPE_API UTextureRenderTarget2DArray* TryGetRenderTarget2DArray() const;
 	/** @return the internal render target's true resolution (can be different than the effective resolution if CurrentRenderTargetParams.bExactDimensions is false) */
-	FIntPoint GetResolution() const;
+	LANDSCAPE_API FIntPoint GetResolution() const;
 	/**@return the internal render target's effective resolution (can be different than the actual resolution if CurrentRenderTargetParams.bExactDimensions is false) while this render target is in use */
-	FIntPoint GetEffectiveResolution() const; 
+	LANDSCAPE_API FIntPoint GetEffectiveResolution() const;
 	/** @return the internal render target's number of slices (can be different than the actual number of slices if CurrentRenderTargetParams.bExactDimensions is false) */
-	int32 GetNumSlices() const;
+	LANDSCAPE_API int32 GetNumSlices() const;
 	/** @return the internal render target's effective number of slices (can be different than the actual number of slices if CurrentRenderTargetParams.bExactDimensions is false) while this render target is in use */
-	int32 GetEffectiveNumSlices() const;
+	LANDSCAPE_API int32 GetEffectiveNumSlices() const;
 	bool IsTexture2DArray() const { return (TryGetRenderTarget2DArray() != nullptr); }
-	ETextureRenderTargetFormat GetFormat() const;
-	FLinearColor GetClearColor() const;
+	LANDSCAPE_API ETextureRenderTargetFormat GetFormat() const;
+	LANDSCAPE_API FLinearColor GetClearColor() const;
 	ERHIAccess GetCurrentState() const { return CurrentState; }
 	const UE::Landscape::FScratchRenderTargetParams& GetCurrentRenderTargetParams() { return CurrentRenderTargetParams; }
 
@@ -169,13 +169,13 @@ public:
 	};
 
 	/** Copies the content of the texture in parameter to the scratch texture (assuming the input texture is in CopySrc state already). Transitions the scratch texture's RHIAccess */
-	void CopyFrom(const FCopyFromTextureParams& InCopyParams);
+	LANDSCAPE_API void CopyFrom(const FCopyFromTextureParams& InCopyParams);
 	/** Copies the content of the scratch texture in parameter to the scratch texture. Transitions both scratch texture's RHIAccess */
-	void CopyFrom(const FCopyFromScratchRenderTargetParams& InCopyParams);
+	LANDSCAPE_API void CopyFrom(const FCopyFromScratchRenderTargetParams& InCopyParams);
 	/** Perform a transition of the scratch texture's RHIAccess */
-	void TransitionTo(ERHIAccess InStateAfter);
+	LANDSCAPE_API void TransitionTo(ERHIAccess InStateAfter);
 	/** Perform a clear of the scratch texture */
-	void Clear();
+	LANDSCAPE_API void Clear();
 	/** @return true if this scratch texture is compatible (and can therefore be used) with the requested render target params */
 	bool IsCompatibleWith(const UE::Landscape::FScratchRenderTargetParams& InParams) const;
 	/** @return true if the scratch texture is currently in use (it has been requested but not released) */
