LandscapePatch: Fix for a crash on editor exit when multiple patches become invalid at the same time.

diff --git a/Engine/Plugins/Experimental/LandscapePatch/Source/LandscapePatch/Private/LandscapePatchEditLayer.cpp b/Engine/Plugins/Experimental/LandscapePatch/Source/LandscapePatch/Private/LandscapePatchEditLayer.cpp
--- a/Engine/Plugins/Experimental/LandscapePatch/Source/LandscapePatch/Private/LandscapePatchEditLayer.cpp
+++ b/Engine/Plugins/Experimental/LandscapePatch/Source/LandscapePatch/Private/LandscapePatchEditLayer.cpp
@@ -209,9 +209,20 @@ void ULandscapePatchEditLayer::UpdateHighestKnownPriority()
 {
 	if (!bPatchListDirty)
 	{
-		HighestKnownPriority = RegisteredPatches.Num() > 0 ?
-			FMath::Max(RegisteredPatches.Last()->GetPriority(), PATCH_PRIORITY_BASE)
-			: PATCH_PRIORITY_BASE;
+		if (RegisteredPatches.Num() == 0)
+		{
+			HighestKnownPriority = PATCH_PRIORITY_BASE;
+		}
+		else if (RegisteredPatches.Last().IsValid())
+		{
+			HighestKnownPriority = FMath::Max(RegisteredPatches.Last()->GetPriority(), PATCH_PRIORITY_BASE);
+		}
+		// If the last patch was invalid, then it seems likely that multiple patches managed
+		//  to become invalid at the same time, and we haven't yet removed the last one while
+		//  processing the NotifyOfPatchRemoval call for a previous one. 
+		// There are a few ways we could handle the situation, but for now we will just leave
+		//  the highest priority unchanged, under the assumption that it will be updated in an
+		//  upcoming NotifyOfPatchRemoval call, when that patch is properly removed.
 	}
 }
 
