Used simpler names for landscape patch edit layer renderer debug names

diff --git a/Engine/Plugins/Experimental/LandscapePatch/Source/LandscapePatch/Private/LandscapeCircleHeightPatch.cpp b/Engine/Plugins/Experimental/LandscapePatch/Source/LandscapePatch/Private/LandscapeCircleHeightPatch.cpp
--- a/Engine/Plugins/Experimental/LandscapePatch/Source/LandscapePatch/Private/LandscapeCircleHeightPatch.cpp
+++ b/Engine/Plugins/Experimental/LandscapePatch/Source/LandscapePatch/Private/LandscapeCircleHeightPatch.cpp
@@ -223,7 +223,7 @@ void ULandscapeCircleHeightPatch::GetRendererStateInfo(const ULandscapeInfo* InL
 
 FString ULandscapeCircleHeightPatch::GetEditLayerRendererDebugName() const
 {
-	return FString::Printf(TEXT("ULandscapeCircleHeightPatch:%s"), *GetFullName());
+	return FString::Printf(TEXT("%s:%s"), *GetOwner()->GetActorNameOrLabel(), *GetName());
 }
 
 TArray<UE::Landscape::EditLayers::FEditLayerRenderItem> ULandscapeCircleHeightPatch::GetRenderItems(const ULandscapeInfo* InLandscapeInfo) const
diff --git a/Engine/Plugins/Experimental/LandscapePatch/Source/LandscapePatch/Private/LandscapeCircleHeightPatchPS.cpp b/Engine/Plugins/Experimental/LandscapePatch/Source/LandscapePatch/Private/LandscapeCircleHeightPatchPS.cpp
--- a/Engine/Plugins/Experimental/LandscapePatch/Source/LandscapePatch/Private/LandscapeCircleHeightPatchPS.cpp
+++ b/Engine/Plugins/Experimental/LandscapePatch/Source/LandscapePatch/Private/LandscapeCircleHeightPatchPS.cpp
@@ -26,7 +26,7 @@ void FLandscapeCircleHeightPatchPS::AddToRenderGraph(FRDGBuilder& GraphBuilder,
 	FPixelShaderUtils::AddFullscreenPass(
 		GraphBuilder,
 		ShaderMap,
-		RDG_EVENT_NAME("LandmassCircleHeightPatch"),
+		RDG_EVENT_NAME("LandscapeCircleHeightPatch"),
 		PixelShader,
 		Parameters,
 		DestinationBounds);
@@ -45,7 +45,7 @@ void FLandscapeCircleVisibilityPatchPS::AddToRenderGraph(FRDGBuilder& GraphBuild
 	FPixelShaderUtils::AddFullscreenPass(
 		GraphBuilder,
 		ShaderMap,
-		RDG_EVENT_NAME("LandmassCircleVisibilityPatch"),
+		RDG_EVENT_NAME("LandscapeCircleVisibilityPatch"),
 		PixelShader,
 		Parameters,
 		DestinationBounds);
diff --git a/Engine/Plugins/Experimental/LandscapePatch/Source/LandscapePatch/Private/LandscapeTexturePatch.cpp b/Engine/Plugins/Experimental/LandscapePatch/Source/LandscapePatch/Private/LandscapeTexturePatch.cpp
--- a/Engine/Plugins/Experimental/LandscapePatch/Source/LandscapePatch/Private/LandscapeTexturePatch.cpp
+++ b/Engine/Plugins/Experimental/LandscapePatch/Source/LandscapePatch/Private/LandscapeTexturePatch.cpp
@@ -1932,7 +1932,7 @@ void ULandscapeTexturePatch::GetRendererStateInfo(const ULandscapeInfo* InLandsc
 
 FString ULandscapeTexturePatch::GetEditLayerRendererDebugName() const
 {
-	return FString::Printf(TEXT("ULandscapeTexturePatch:%s"), *GetFullName());
+	return FString::Printf(TEXT("%s:%s"), *GetOwner()->GetActorNameOrLabel(), *GetName());
 }
 
 TArray<UE::Landscape::EditLayers::FEditLayerRenderItem> ULandscapeTexturePatch::GetRenderItems(const ULandscapeInfo* InLandscapeInfo) const
