LandscapePatch: Fix ensure when edit layer is deleted with patches inside.

diff --git a/Engine/Plugins/Experimental/LandscapePatch/Source/LandscapePatch/Private/LandscapePatchComponent.cpp b/Engine/Plugins/Experimental/LandscapePatch/Source/LandscapePatch/Private/LandscapePatchComponent.cpp
--- a/Engine/Plugins/Experimental/LandscapePatch/Source/LandscapePatch/Private/LandscapePatchComponent.cpp
+++ b/Engine/Plugins/Experimental/LandscapePatch/Source/LandscapePatch/Private/LandscapePatchComponent.cpp
@@ -1064,14 +1064,9 @@ void ULandscapePatchComponent::FixBindings()
 void ULandscapePatchComponent::NotifyOfBoundLayerDeletion(ULandscapePatchEditLayer* LayerIn)
 {
 	// If we're getting this call, we are presumably bound through an edit layer, not patch manager
-	if (!ensure(!PatchManager.IsValid()))
-	{
-		ResetEditLayer();
-		return;
-	}
+	ensure(!PatchManager.IsValid());
 
-	// This should update our cached pointer and displayed layer name
-	BindToEditLayer(EditLayerGuid);
+	ResetEditLayer();
 }
 #endif
 
diff --git a/Engine/Plugins/Experimental/LandscapePatch/Source/LandscapePatch/Private/LandscapePatchEditLayer.cpp b/Engine/Plugins/Experimental/LandscapePatch/Source/LandscapePatch/Private/LandscapePatchEditLayer.cpp
--- a/Engine/Plugins/Experimental/LandscapePatch/Source/LandscapePatch/Private/LandscapePatchEditLayer.cpp
+++ b/Engine/Plugins/Experimental/LandscapePatch/Source/LandscapePatch/Private/LandscapePatchEditLayer.cpp
@@ -280,7 +280,11 @@ void ULandscapePatchEditLayer::OnLayerRemoved()
 	// TODO: If we end up keeping this pointer, it should probably be reset in the base class implementation
 	// of OnLayerRemoved.
 	OwningLandscape.Reset();
-	for (TSoftObjectPtr<ULandscapePatchComponent> PatchSoft : RegisteredPatches)
+
+	// Iterate through a copy so that patches can deregister themselves in NotifyOfBoundLayerDeletion
+	//  without messing up our iteration.
+	TArray<TSoftObjectPtr<ULandscapePatchComponent>> PatchesCopy = RegisteredPatches;
+	for (TSoftObjectPtr<ULandscapePatchComponent> PatchSoft : PatchesCopy)
 	{
 		ULandscapePatchComponent* Patch = PatchSoft.Get();
 		if (ShouldPatchBeIncludedInList(Patch))
