LandscapePatch: Fix blueprint contained patches losing their internal data on undo. Make patches more consistently start with a cleared texture. Add a couple extra test patches to the test map.

#virtualized

diff --git a/Engine/Plugins/Experimental/LandscapePatch/Source/LandscapePatch/Private/LandscapeTextureBackedRenderTarget.cpp b/Engine/Plugins/Experimental/LandscapePatch/Source/LandscapePatch/Private/LandscapeTextureBackedRenderTarget.cpp
--- a/Engine/Plugins/Experimental/LandscapePatch/Source/LandscapePatch/Private/LandscapeTextureBackedRenderTarget.cpp
+++ b/Engine/Plugins/Experimental/LandscapePatch/Source/LandscapePatch/Private/LandscapeTextureBackedRenderTarget.cpp
@@ -3,6 +3,7 @@
 #include "LandscapeTextureBackedRenderTarget.h"
 #include "Engine/Texture2D.h"
 #include "Engine/World.h"
+#include "LandscapeDataAccess.h"
 #include "LandscapePatchUtil.h" // CopyTextureOnRenderThread
 #include "LandscapePatchLogging.h" 
 #include "RenderGraphBuilder.h"
@@ -33,6 +34,23 @@ namespace LandscapeTextureBackedRenderTargetLocals
 
 		return Texture;
 	}
+
+	// Copied from CameraCalibrationUtilsPrivate.cpp
+	void ClearTexture(UTexture2D* Texture, FColor ClearColor)
+	{
+		if (Texture)
+		{
+			TArray<FColor> Pixels;
+			Pixels.Init(ClearColor, Texture->GetSizeX() * Texture->GetSizeY());
+
+			void* TextureData = Texture->GetPlatformData()->Mips[0].BulkData.Lock(LOCK_READ_WRITE);
+
+			FMemory::Memcpy(TextureData, Pixels.GetData(), Pixels.Num() * sizeof(FColor));
+
+			Texture->GetPlatformData()->Mips[0].BulkData.Unlock();
+			Texture->UpdateResource();
+		}
+	}
 #endif // WITH_EDITOR
 }
 
@@ -62,7 +80,7 @@ void ULandscapeTextureBackedRenderTargetBase::PostLoad()
 	{
 		InternalTexture->ConditionalPostLoad();
 
-		ReinitializeRenderTarget();
+		ReinitializeRenderTarget(/*bClear*/ false);
 
 		CopyBackFromInternalTexture();
 	}
@@ -102,7 +120,19 @@ void ULandscapeTextureBackedRenderTargetBase::PostEditImport()
 
 	if (!bUseInternalTextureOnly && IsValid(InternalTexture))
 	{
-		ReinitializeRenderTarget();
+		ReinitializeRenderTarget(/*bClear*/ false);
+
+		CopyBackFromInternalTexture();
+	}
+}
+
+void ULandscapeTextureBackedRenderTargetBase::PostEditUndo()
+{
+	Super::PostEditUndo();
+
+	if (!bUseInternalTextureOnly && IsValid(InternalTexture))
+	{
+		ReinitializeRenderTarget(/*bClear*/ false);
 
 		CopyBackFromInternalTexture();
 	}
@@ -136,7 +166,7 @@ void ULandscapeTextureBackedRenderTargetBase::SetUseInternalTextureOnly(bool bUs
 
 		if (IsValid(InternalTexture)) // if initialized
 		{
-			ReinitializeRenderTarget();
+			ReinitializeRenderTarget(/*bClear*/ !bCopyExisting);
 
 			if (bCopyExisting)
 			{
@@ -192,7 +222,7 @@ void ULandscapeTextureBackedRenderTargetBase::Initialize()
 	}
 	else
 	{
-		ReinitializeRenderTarget();
+		ReinitializeRenderTarget(/*bClear*/ true);
 	}
 #endif // WITH_EDITOR
 }
@@ -239,6 +269,20 @@ ETextureSourceFormat ULandscapeWeightTextureBackedRenderTarget::GetInternalTextu
 	return bUseAlphaChannel ? ETextureSourceFormat::TSF_BGRA8 : ETextureSourceFormat::TSF_G8;
 }
 
+void ULandscapeWeightTextureBackedRenderTarget::Initialize()
+{
+#if WITH_EDITOR
+	using namespace LandscapeTextureBackedRenderTargetLocals;
+
+	Super::Initialize();
+
+	if (bUseInternalTextureOnly && ensure(InternalTexture))
+	{
+		ClearTexture(InternalTexture, FColor::White);
+	}
+#endif // WITH_EDITOR
+}
+
 void ULandscapeWeightTextureBackedRenderTarget::CopyToInternalTexture()
 {
 #if WITH_EDITOR
@@ -359,6 +403,21 @@ void ULandscapeHeightTextureBackedRenderTarget::SetFormat(ETextureRenderTargetFo
 #endif // WITH_EDITOR
 }
 
+void ULandscapeHeightTextureBackedRenderTarget::Initialize()
+{
+#if WITH_EDITOR
+	using namespace LandscapeTextureBackedRenderTargetLocals;
+	Super::Initialize();
+
+	const FColor LandscapeNativeMidHeightColor = LandscapeDataAccess::PackHeight(LandscapeDataAccess::MidValue);
+
+	if (bUseInternalTextureOnly && ensure(InternalTexture))
+	{
+		ClearTexture(InternalTexture, LandscapeNativeMidHeightColor);
+	}
+#endif // WITH_EDITOR
+}
+
 void ULandscapeHeightTextureBackedRenderTarget::CopyToInternalTexture()
 {
 #if WITH_EDITOR
@@ -524,7 +583,7 @@ bool ULandscapeTextureBackedRenderTargetBase::IsCopyingBackAndForthAllowed()
 		&& FApp::CanEverRender();
 }
 
-void ULandscapeTextureBackedRenderTargetBase::ReinitializeRenderTarget()
+void ULandscapeTextureBackedRenderTargetBase::ReinitializeRenderTarget(bool bClear)
 {
 	if (!IsValid(RenderTarget))
 	{
@@ -538,5 +597,5 @@ void ULandscapeTextureBackedRenderTargetBase::ReinitializeRenderTarget()
 	}
 	RenderTarget->RenderTargetFormat = GetRenderTargetFormat();
 	RenderTarget->InitAutoFormat(SizeX, SizeY);
-	RenderTarget->UpdateResourceImmediate(false);
+	RenderTarget->UpdateResourceImmediate(bClear);
 }
diff --git a/Engine/Plugins/Experimental/LandscapePatch/Source/LandscapePatch/Private/LandscapeTexturePatch.cpp b/Engine/Plugins/Experimental/LandscapePatch/Source/LandscapePatch/Private/LandscapeTexturePatch.cpp
--- a/Engine/Plugins/Experimental/LandscapePatch/Source/LandscapePatch/Private/LandscapeTexturePatch.cpp
+++ b/Engine/Plugins/Experimental/LandscapePatch/Source/LandscapePatch/Private/LandscapeTexturePatch.cpp
@@ -1116,6 +1116,60 @@ void ULandscapeTexturePatch::GetRenderDependencies(TSet<UObject*>& OutDependenci
 	}
 }
 
+TStructOnScope<FActorComponentInstanceData> ULandscapeTexturePatch::GetComponentInstanceData() const
+{
+	// There are currently various issues with blueprints and instanced sub objects, and
+	//  one of them causes undo to be severely broken for transactable instanced objects
+	//  inside a blueprint actor component: UE-225445
+	// As it happens, one workaround is to not have the objects be transactable. So for
+	//  now, we temporarily make all instanced objects not transactable while doing instance
+	//  data serialization (when it theoretically shouldn't matter anyway).
+
+	auto SetObjectTransactionalFlag = [](UObject* Object, bool bOn)
+	{
+		if (!Object)
+		{
+			return;
+		}
+		if (bOn)
+		{
+			Object->SetFlags(RF_Transactional);
+		}
+		else
+		{
+			Object->ClearFlags(RF_Transactional);
+		}
+	};
+	auto SetInternalDataTransactionalFlags = [&SetObjectTransactionalFlag](TObjectPtr<ULandscapeTextureBackedRenderTargetBase> InternalData, bool bOn)
+	{
+		if (!InternalData)
+		{
+			return;
+		}
+		SetObjectTransactionalFlag(InternalData, bOn);
+		SetObjectTransactionalFlag(InternalData->GetRenderTarget(), bOn);
+		SetObjectTransactionalFlag(InternalData->GetInternalTexture(), bOn);
+	};
+	auto SetAllInternalDataTransactionalFlags = [this, &SetObjectTransactionalFlag, &SetInternalDataTransactionalFlags](bool bOn)
+	{
+		SetInternalDataTransactionalFlags(HeightInternalData, bOn);
+		for (const TObjectPtr<ULandscapeWeightPatchTextureInfo>& WeightPatch : WeightPatches)
+		{
+			if (IsValid(WeightPatch))
+			{
+				SetObjectTransactionalFlag(WeightPatch, bOn);
+				SetInternalDataTransactionalFlags(WeightPatch->InternalData, bOn);
+			}
+		}
+	};
+	
+	SetAllInternalDataTransactionalFlags(false);
+	TStructOnScope<FActorComponentInstanceData> ToReturn = Super::GetComponentInstanceData();
+	SetAllInternalDataTransactionalFlags(true);
+
+	return ToReturn;
+}
+
 #endif
 
 void ULandscapeTexturePatch::SnapToLandscape()
diff --git a/Engine/Plugins/Experimental/LandscapePatch/Source/LandscapePatch/Public/LandscapeTextureBackedRenderTarget.h b/Engine/Plugins/Experimental/LandscapePatch/Source/LandscapePatch/Public/LandscapeTextureBackedRenderTarget.h
--- a/Engine/Plugins/Experimental/LandscapePatch/Source/LandscapePatch/Public/LandscapeTextureBackedRenderTarget.h
+++ b/Engine/Plugins/Experimental/LandscapePatch/Source/LandscapePatch/Public/LandscapeTextureBackedRenderTarget.h
@@ -25,9 +25,6 @@ class ULandscapeTextureBackedRenderTargetBase : public UObject
 
 public:
 
-	UPROPERTY()
-	TObjectPtr<UTextureRenderTarget2D> PostLoadRT = nullptr;
-
 	virtual UTextureRenderTarget2D* GetRenderTarget() { return RenderTarget; }
 	virtual UTexture2D* GetInternalTexture() { return InternalTexture; }
 
@@ -50,6 +47,7 @@ public:
 	virtual void PreDuplicate(FObjectDuplicationParameters& DupParams) override;
 	virtual void ExportCustomProperties(FOutputDevice& Out, uint32 Indent) override;
 	virtual void PostEditImport() override;
+	virtual void PostEditUndo() override;
 #endif // WITH_EDITOR
 
 protected:
@@ -60,7 +58,7 @@ protected:
 	virtual ETextureSourceFormat GetInternalTextureFormat() PURE_VIRTUAL(ULandscapeTextureBackedRenderTargetBase::GetInternalTextureFormat, return ETextureSourceFormat::TSF_G8;);
 	virtual ETextureRenderTargetFormat GetRenderTargetFormat() PURE_VIRTUAL(ULandscapeTextureBackedRenderTargetBase::GetRenderTargetFormat, return ETextureRenderTargetFormat::RTF_R8;);
 
-	void ReinitializeRenderTarget();
+	void ReinitializeRenderTarget(bool bClear);
 
 	UPROPERTY(VisibleAnywhere, Category = InternalData)
 	TObjectPtr<UTexture2D> InternalTexture = nullptr;
@@ -87,6 +85,7 @@ public:
 
 	virtual void SetUseAlphaChannel(bool bUseAlphaChannelIn);
 
+	virtual void Initialize() override;
 	virtual void CopyToInternalTexture() override;
 	virtual void CopyBackFromInternalTexture() override;
 protected:
@@ -106,6 +105,7 @@ public:
 
 	virtual void SetFormat(ETextureRenderTargetFormat FormatToUse);
 
+	virtual void Initialize() override;
 	virtual void CopyToInternalTexture() override;
 	virtual void CopyBackFromInternalTexture() override;
 
diff --git a/Engine/Plugins/Experimental/LandscapePatch/Source/LandscapePatch/Public/LandscapeTexturePatch.h b/Engine/Plugins/Experimental/LandscapePatch/Source/LandscapePatch/Public/LandscapeTexturePatch.h
--- a/Engine/Plugins/Experimental/LandscapePatch/Source/LandscapePatch/Public/LandscapeTexturePatch.h
+++ b/Engine/Plugins/Experimental/LandscapePatch/Source/LandscapePatch/Public/LandscapeTexturePatch.h
@@ -240,6 +240,9 @@ public:
 	virtual bool CanAffectVisibilityLayer() const override;
 	virtual void GetRenderDependencies(TSet<UObject*>& OutDependencies) const override;
 
+	// UActorComponent
+	virtual TStructOnScope<FActorComponentInstanceData> GetComponentInstanceData() const override;
+
 	// UObject
 	virtual void PostEditChangeProperty(FPropertyChangedEvent& PropertyChangedEvent) override;
 #endif
