Cleaned up LandscapePatch public API

diff --git a/Engine/Plugins/Experimental/LandscapePatch/Source/LandscapePatch/Private/LandscapeCircleHeightPatchPS.h b/Engine/Plugins/Experimental/LandscapePatch/Source/LandscapePatch/Private/LandscapeCircleHeightPatchPS.h
new file mode 100644
--- /dev/null
+++ b/Engine/Plugins/Experimental/LandscapePatch/Source/LandscapePatch/Private/LandscapeCircleHeightPatchPS.h
@@ -0,0 +1,51 @@
+// Copyright Epic Games, Inc. All Rights Reserved.
+
+#pragma once
+
+#include "GlobalShader.h"
+#include "ShaderParameterStruct.h"
+
+/**
+ * Shader that applies a circle height patch to a landscape heightmap.
+ */
+class FLandscapeCircleHeightPatchPS : public FGlobalShader
+{
+	DECLARE_GLOBAL_SHADER(FLandscapeCircleHeightPatchPS);
+	SHADER_USE_PARAMETER_STRUCT(FLandscapeCircleHeightPatchPS, FGlobalShader);
+
+public:
+	BEGIN_SHADER_PARAMETER_STRUCT(FParameters, )
+		SHADER_PARAMETER_RDG_TEXTURE_SRV(Texture2D<float4>, InSourceTexture) // Our input texture
+		// Offset of the source heightmap relative to the 0,0 location in the destination heightmap, because
+		// the source is likely to be a copied region from some inner part of the destination. This is basically
+		// a SourceHeightmapToDestinationHeightmap coordinate transformation, except that it is always a simple 
+		// integer translation.
+		SHADER_PARAMETER(FIntPoint, InSourceTextureOffset)
+		SHADER_PARAMETER(FVector3f, InCenter)
+		SHADER_PARAMETER(float, InRadius)
+		SHADER_PARAMETER(float, InFalloff)
+		RENDER_TARGET_BINDING_SLOTS() // Holds our output
+	END_SHADER_PARAMETER_STRUCT()
+
+	static bool ShouldCompilePermutation(const FGlobalShaderPermutationParameters& Parameters);
+	static void ModifyCompilationEnvironment(const FGlobalShaderPermutationParameters& Parameters, FShaderCompilerEnvironment& OutEnvironment);
+
+	static void AddToRenderGraph(FRDGBuilder& GraphBuilder, FParameters* InParameters, const FIntRect& DestinationBounds);
+};
+
+/**
+ * Shader that applies a circle patch to a landscape visibility layer.
+ */
+class FLandscapeCircleVisibilityPatchPS : public FLandscapeCircleHeightPatchPS
+{
+	DECLARE_GLOBAL_SHADER(FLandscapeCircleVisibilityPatchPS);
+
+public:
+	FLandscapeCircleVisibilityPatchPS() = default;
+	FLandscapeCircleVisibilityPatchPS(const ShaderMetaType::CompiledShaderInitializerType& Initializer)
+		: FLandscapeCircleHeightPatchPS(Initializer)
+	{}
+
+	static void ModifyCompilationEnvironment(const FGlobalShaderPermutationParameters& Parameters, FShaderCompilerEnvironment& OutEnvironment);
+	static void AddToRenderGraph(FRDGBuilder& GraphBuilder, FParameters* InParameters, const FIntRect& DestinationBounds);
+};
diff --git a/Engine/Plugins/Experimental/LandscapePatch/Source/LandscapePatch/Private/LandscapeTexturePatch.cpp b/Engine/Plugins/Experimental/LandscapePatch/Source/LandscapePatch/Private/LandscapeTexturePatch.cpp
--- a/Engine/Plugins/Experimental/LandscapePatch/Source/LandscapePatch/Private/LandscapeTexturePatch.cpp
+++ b/Engine/Plugins/Experimental/LandscapePatch/Source/LandscapePatch/Private/LandscapeTexturePatch.cpp
@@ -10,6 +10,7 @@
 #include "LandscapePatchLogging.h"
 #include "LandscapePatchManager.h"
 #include "LandscapePatchUtil.h" // CopyTextureOnRenderThread
+#include "LandscapeTexturePatchPS.h"
 #include "LandscapeUtils.h" // IsVisibilityLayer
 #include "MathUtil.h"
 #include "RHIStaticStates.h"
@@ -324,7 +325,7 @@ UTextureRenderTarget2D* ULandscapeTexturePatch::ApplyToHeightmap(UTextureRenderT
 
 	// Go ahead and pack everything into a copy of the param struct so we don't have to capture everything
 	// individually in the lambda below.
-	FApplyLandscapeTextureHeightPatchPS::FParameters ShaderParamsToCopy;
+	FApplyLandscapeTextureHeightPatchPSParameters ShaderParamsToCopy;
 	FIntRect DestinationBounds;
 	GetHeightShaderParams(LandscapeHeightmapToWorld, FIntPoint(Patch->GetSizeX(), Patch->GetSizeY()), FIntPoint(InCombinedResult->SizeX, InCombinedResult->SizeY), ShaderParamsToCopy, DestinationBounds);
 
@@ -430,7 +431,7 @@ void ULandscapeTexturePatch::ApplyToWeightmap(ULandscapeWeightPatchTextureInfo*
 
 	// Go ahead and pack everything into a copy of the param struct so we don't have to capture everything
 	// individually in the lambda below.
-	FApplyLandscapeTextureWeightPatchPS::FParameters ShaderParamsToCopy;
+	FApplyLandscapeTextureWeightPatchPSParameters ShaderParamsToCopy;
 	FIntRect DestinationBounds;
 
 	GetWeightShaderParams(LandscapeHeightmapToWorld, FIntPoint(Patch->GetSizeX(), Patch->GetSizeY()), 
@@ -553,7 +554,7 @@ void ULandscapeTexturePatch::GetCommonShaderParams(const FTransform& LandscapeHe
 
 void ULandscapeTexturePatch::GetHeightShaderParams(const FTransform& LandscapeHeightmapToWorldIn,
 	const FIntPoint& SourceResolutionIn, const FIntPoint& DestinationResolutionIn,
-	UE::Landscape::FApplyLandscapeTextureHeightPatchPS::FParameters& ParamsOut,
+	UE::Landscape::FApplyLandscapeTextureHeightPatchPSParameters& ParamsOut,
 	FIntRect& DestinationBoundsOut) const
 {
 	using namespace UE::Landscape;
@@ -623,7 +624,7 @@ void ULandscapeTexturePatch::GetHeightShaderParams(const FTransform& LandscapeHe
 void ULandscapeTexturePatch::GetWeightShaderParams(
 	const FTransform& LandscapeHeightmapToWorldIn, const FIntPoint& SourceResolutionIn,
 	const FIntPoint& DestinationResolutionIn, const ULandscapeWeightPatchTextureInfo* WeightPatchInfo, 
-	UE::Landscape::FApplyLandscapeTextureWeightPatchPS::FParameters& ParamsOut, 
+	UE::Landscape::FApplyLandscapeTextureWeightPatchPSParameters& ParamsOut,
 	FIntRect& DestinationBoundsOut) const
 {
 	using namespace UE::Landscape;
diff --git a/Engine/Plugins/Experimental/LandscapePatch/Source/LandscapePatch/Private/LandscapeTexturePatchPS.cpp b/Engine/Plugins/Experimental/LandscapePatch/Source/LandscapePatch/Private/LandscapeTexturePatchPS.cpp
--- a/Engine/Plugins/Experimental/LandscapePatch/Source/LandscapePatch/Private/LandscapeTexturePatchPS.cpp
+++ b/Engine/Plugins/Experimental/LandscapePatch/Source/LandscapePatch/Private/LandscapeTexturePatchPS.cpp
@@ -3,10 +3,9 @@
 #include "LandscapeTexturePatchPS.h"
 
 #include "LandscapeUtils.h"
+#include "LandscapePatchUtil.h"
 #include "PixelShaderUtils.h"
 
-#include UE_INLINE_GENERATED_CPP_BY_NAME(LandscapeTexturePatchPS)
-
 namespace UE::Landscape
 {
 
@@ -24,10 +23,10 @@ void FApplyLandscapeTextureHeightPatchPS::ModifyCompilationEnvironment(const FGl
 	OutEnvironment.SetDefine(TEXT("APPLY_PATCH_ALPHA_FLAG"), static_cast<uint8>(FApplyLandscapeTextureHeightPatchPS::EFlags::ApplyPatchAlpha));
 	OutEnvironment.SetDefine(TEXT("INPUT_IS_PACKED_HEIGHT_FLAG"), static_cast<uint8>(FApplyLandscapeTextureHeightPatchPS::EFlags::InputIsPackedHeight));
 
-	OutEnvironment.SetDefine(TEXT("ADDITIVE_MODE"), static_cast<uint8>(FApplyLandscapeTextureHeightPatchPS::EBlendMode::Additive));
-	OutEnvironment.SetDefine(TEXT("ALPHA_BLEND_MODE"), static_cast<uint8>(FApplyLandscapeTextureHeightPatchPS::EBlendMode::AlphaBlend));
-	OutEnvironment.SetDefine(TEXT("MIN_MODE"), static_cast<uint8>(FApplyLandscapeTextureHeightPatchPS::EBlendMode::Min));
-	OutEnvironment.SetDefine(TEXT("MAX_MODE"), static_cast<uint8>(FApplyLandscapeTextureHeightPatchPS::EBlendMode::Max));
+	OutEnvironment.SetDefine(TEXT("ALPHA_BLEND_MODE"), static_cast<uint8>(ELandscapeTexturePatchBlendMode::AlphaBlend));
+	OutEnvironment.SetDefine(TEXT("ADDITIVE_MODE"), static_cast<uint8>(ELandscapeTexturePatchBlendMode::Additive));
+	OutEnvironment.SetDefine(TEXT("MIN_MODE"), static_cast<uint8>(ELandscapeTexturePatchBlendMode::Min));
+	OutEnvironment.SetDefine(TEXT("MAX_MODE"), static_cast<uint8>(ELandscapeTexturePatchBlendMode::Max));
 }
 
 void FApplyLandscapeTextureHeightPatchPS::AddToRenderGraph(FRDGBuilder& GraphBuilder, FParameters* InParameters, const FIntRect& DestinationBounds)
@@ -183,13 +182,13 @@ void FApplyLandscapeTextureWeightPatchPS::ModifyCompilationEnvironment(const FGl
 	OutEnvironment.SetDefine(TEXT("APPLY_WEIGHT_PATCH"), 1);
 
 	// Make our flag choices match in the shader.
-	OutEnvironment.SetDefine(TEXT("RECTANGULAR_FALLOFF_FLAG"), static_cast<uint8>(FApplyLandscapeTextureHeightPatchPS::EFlags::RectangularFalloff));
-	OutEnvironment.SetDefine(TEXT("APPLY_PATCH_ALPHA_FLAG"), static_cast<uint8>(FApplyLandscapeTextureHeightPatchPS::EFlags::ApplyPatchAlpha));
+	OutEnvironment.SetDefine(TEXT("RECTANGULAR_FALLOFF_FLAG"), static_cast<uint8>(FApplyLandscapeTextureWeightPatchPS::EFlags::RectangularFalloff));
+	OutEnvironment.SetDefine(TEXT("APPLY_PATCH_ALPHA_FLAG"), static_cast<uint8>(FApplyLandscapeTextureWeightPatchPS::EFlags::ApplyPatchAlpha));
 
-	OutEnvironment.SetDefine(TEXT("ADDITIVE_MODE"), static_cast<uint8>(FApplyLandscapeTextureHeightPatchPS::EBlendMode::Additive));
-	OutEnvironment.SetDefine(TEXT("ALPHA_BLEND_MODE"), static_cast<uint8>(FApplyLandscapeTextureHeightPatchPS::EBlendMode::AlphaBlend));
-	OutEnvironment.SetDefine(TEXT("MIN_MODE"), static_cast<uint8>(FApplyLandscapeTextureHeightPatchPS::EBlendMode::Min));
-	OutEnvironment.SetDefine(TEXT("MAX_MODE"), static_cast<uint8>(FApplyLandscapeTextureHeightPatchPS::EBlendMode::Max));
+	OutEnvironment.SetDefine(TEXT("ALPHA_BLEND_MODE"), static_cast<uint8>(ELandscapeTexturePatchBlendMode::AlphaBlend));
+	OutEnvironment.SetDefine(TEXT("ADDITIVE_MODE"), static_cast<uint8>(ELandscapeTexturePatchBlendMode::Additive));
+	OutEnvironment.SetDefine(TEXT("MIN_MODE"), static_cast<uint8>(ELandscapeTexturePatchBlendMode::Min));
+	OutEnvironment.SetDefine(TEXT("MAX_MODE"), static_cast<uint8>(ELandscapeTexturePatchBlendMode::Max));
 }
 
 void FApplyLandscapeTextureWeightPatchPS::AddToRenderGraph(FRDGBuilder& GraphBuilder, FParameters* InParameters, const FIntRect& DestinationBounds)
diff --git a/Engine/Plugins/Experimental/LandscapePatch/Source/LandscapePatch/Private/LandscapeTexturePatchPS.h b/Engine/Plugins/Experimental/LandscapePatch/Source/LandscapePatch/Private/LandscapeTexturePatchPS.h
new file mode 100644
--- /dev/null
+++ b/Engine/Plugins/Experimental/LandscapePatch/Source/LandscapePatch/Private/LandscapeTexturePatchPS.h
@@ -0,0 +1,272 @@
+// Copyright Epic Games, Inc. All Rights Reserved.
+
+#pragma once
+
+
+#include "GlobalShader.h"
+#include "ShaderParameterStruct.h"
+#include "LandscapeTexturePatch.h"
+
+
+struct FLandscapeHeightPatchConvertToNativeParams;
+
+namespace UE::Landscape
+{
+
+/**
+ * Shader that applies a texture-based height patch to a landscape heightmap.
+ */
+class FApplyLandscapeTextureHeightPatchPS : public FGlobalShader
+{
+	DECLARE_GLOBAL_SHADER(FApplyLandscapeTextureHeightPatchPS);
+	SHADER_USE_PARAMETER_STRUCT(FApplyLandscapeTextureHeightPatchPS, FGlobalShader);
+
+public:
+
+	// Flags that get packed into a bitfield because we're not allowed to use bool shader parameters:
+	enum class EFlags : uint8
+	{
+		None = 0,
+
+		// When false, falloff is circular.
+		RectangularFalloff = 1 << 0,
+
+		// When true, the texture alpha channel is considered for blending (in addition to falloff, if nonzero)
+		ApplyPatchAlpha = 1 << 1,
+
+		// When false, the input is directly interpreted as being the height value to process. When true, the height
+		// is unpacked from the red and green channels to make a 16 bit int.
+		InputIsPackedHeight = 1 << 2
+	};
+
+	// TODO: We could consider exposing an additional global alpha setting that we can use to pass in the given
+	// edit layer alpha value... On the other hand, we currently don't bother doing this in any existing blueprint
+	// brushes, and it would be hard to support in a way that doesn't require each blueprint brush to respect it
+	// individually... Not clear whether this is something worth doing yet.
+
+	BEGIN_SHADER_PARAMETER_STRUCT(FParameters, )
+		SHADER_PARAMETER_RDG_TEXTURE_SRV(Texture2D<float4>, InSourceHeightmap)
+		// Offset of the source heightmap relative to the 0,0 location in the destination heightmap, because
+		// the source is likely to be a copied region from some inner part of the destination. This is basically
+		// a SourceHeightmapToDestinationHeightmap coordinate transformation, except that it is always a simple 
+		// integer translation.
+		SHADER_PARAMETER(FIntPoint, InSourceHeightmapOffset)
+		SHADER_PARAMETER_RDG_TEXTURE_SRV(Texture2D<float4>, InHeightPatch)
+		SHADER_PARAMETER_SAMPLER(SamplerState, InHeightPatchSampler)
+		SHADER_PARAMETER(FMatrix44f, InHeightmapToPatch)
+		// Value in patch that corresponds to the landscape mid value, which is our "0 height".
+		SHADER_PARAMETER(float, InZeroInEncoding)
+		// Scale to apply to source values relative to the value that represents 0 height.
+		SHADER_PARAMETER(float, InHeightScale)
+		// Offset to apply to height result after applying height scale
+		SHADER_PARAMETER(float, InHeightOffset)
+		// Amount of the patch edge to not apply in UV space. Generally set to 0.5/Dimensions to avoid applying
+		// the edge half-pixels.
+		SHADER_PARAMETER(FVector2f, InEdgeUVDeadBorder)
+		// In world units, the size of the margin across which the alpha falls from 1 to 0
+		SHADER_PARAMETER(float, InFalloffWorldMargin)
+		// Size of the patch in world units (used for falloff)
+		SHADER_PARAMETER(FVector2f, InPatchWorldDimensions)
+		SHADER_PARAMETER(uint32, InBlendMode)
+		// Some combination of the flags (see constants above).
+		SHADER_PARAMETER(uint32, InFlags)
+
+		RENDER_TARGET_BINDING_SLOTS() // Holds our output
+	END_SHADER_PARAMETER_STRUCT()
+
+	static bool ShouldCompilePermutation(const FGlobalShaderPermutationParameters& Parameters);
+	static void ModifyCompilationEnvironment(const FGlobalShaderPermutationParameters& InParameters, FShaderCompilerEnvironment& OutEnvironment);
+
+	static void AddToRenderGraph(FRDGBuilder& GraphBuilder, FParameters* InParameters, const FIntRect& DestinationBounds);
+};
+
+// We need an alias for FApplyLandscapeTextureHeightPatchPS::FParameters so that it can be forward-declared in another header file (cannot forward-declare internal classes)
+class FApplyLandscapeTextureHeightPatchPSParameters : public FApplyLandscapeTextureHeightPatchPS::FParameters {};
+
+ENUM_CLASS_FLAGS(FApplyLandscapeTextureHeightPatchPS::EFlags);
+
+/**
+ * Simple shader that just offsets each height value in a height patch by a constant.
+ */
+class FOffsetHeightmapPS : public FGlobalShader
+{
+	DECLARE_GLOBAL_SHADER(FOffsetHeightmapPS);
+	SHADER_USE_PARAMETER_STRUCT(FOffsetHeightmapPS, FGlobalShader);
+
+public:
+
+	BEGIN_SHADER_PARAMETER_STRUCT(FParameters, )
+		SHADER_PARAMETER_RDG_TEXTURE_SRV(Texture2D<float4>, InHeightmap)
+		SHADER_PARAMETER(float, InHeightOffset)
+		RENDER_TARGET_BINDING_SLOTS() // Holds our output
+	END_SHADER_PARAMETER_STRUCT()
+
+	static bool ShouldCompilePermutation(const FGlobalShaderPermutationParameters& Parameters);
+	static void ModifyCompilationEnvironment(const FGlobalShaderPermutationParameters& InParameters, FShaderCompilerEnvironment& OutEnvironment);
+
+	static void AddToRenderGraph(FRDGBuilder& GraphBuilder, FParameters* InParameters);
+};
+
+/**
+ * Shader that converts a texture stored in some other encoding (where height is in the R channel) to the 
+ * landscape "native" encoding, where height is stored as a 16 bit int split across the R and G channels.
+ * This is not perfectly reversible (in case of clamping and due to rounding), but it lets us store the
+ * texture in the way that it would be applied to the landscape (usually).
+ */
+class FConvertToNativeLandscapePatchPS : public FGlobalShader
+{
+	DECLARE_GLOBAL_SHADER(FConvertToNativeLandscapePatchPS);
+	SHADER_USE_PARAMETER_STRUCT(FConvertToNativeLandscapePatchPS, FGlobalShader);
+
+public:
+
+	BEGIN_SHADER_PARAMETER_STRUCT(FParameters, )
+		SHADER_PARAMETER_RDG_TEXTURE_SRV(Texture2D<float4>, InHeightmap)
+		SHADER_PARAMETER(float, InZeroInEncoding)
+		SHADER_PARAMETER(float, InHeightScale)
+		SHADER_PARAMETER(float, InHeightOffset)
+		RENDER_TARGET_BINDING_SLOTS() // Holds our output
+	END_SHADER_PARAMETER_STRUCT()
+
+	static bool ShouldCompilePermutation(const FGlobalShaderPermutationParameters& Parameters);
+	static void ModifyCompilationEnvironment(const FGlobalShaderPermutationParameters& InParameters, FShaderCompilerEnvironment& OutEnvironment);
+
+	static void AddToRenderGraph(FRDGBuilder& GraphBuilder, FRDGTextureRef SourceTexture,
+		FRDGTextureRef DestinationTexture, const FLandscapeHeightPatchConvertToNativeParams& Params);
+};
+
+/**
+ * Shader that undoes the conversion done by FConvertToNativeLandscapePatchPS (to the extent possible, since
+ * rounding/clamping makes it not perfectly recoverable).
+ */
+class FConvertBackFromNativeLandscapePatchPS : public FGlobalShader
+{
+	DECLARE_GLOBAL_SHADER(FConvertBackFromNativeLandscapePatchPS);
+	SHADER_USE_PARAMETER_STRUCT(FConvertBackFromNativeLandscapePatchPS, FGlobalShader);
+
+public:
+
+	BEGIN_SHADER_PARAMETER_STRUCT(FParameters, )
+		SHADER_PARAMETER_RDG_TEXTURE_SRV(Texture2D<float4>, InHeightmap)
+		SHADER_PARAMETER(float, InZeroInEncoding)
+		SHADER_PARAMETER(float, InHeightScale)
+		SHADER_PARAMETER(float, InHeightOffset)
+		RENDER_TARGET_BINDING_SLOTS() // Holds our output
+	END_SHADER_PARAMETER_STRUCT()
+
+	static bool ShouldCompilePermutation(const FGlobalShaderPermutationParameters& Parameters);
+	static void ModifyCompilationEnvironment(const FGlobalShaderPermutationParameters& InParameters, FShaderCompilerEnvironment& OutEnvironment);
+
+	static void AddToRenderGraph(FRDGBuilder& GraphBuilder, FRDGTextureRef SourceTexture,
+		FRDGTextureRef DestinationTexture, const FLandscapeHeightPatchConvertToNativeParams& Params);
+};
+
+class FApplyLandscapeTextureWeightPatchPS : public FGlobalShader
+{
+	DECLARE_GLOBAL_SHADER(FApplyLandscapeTextureWeightPatchPS);
+	SHADER_USE_PARAMETER_STRUCT(FApplyLandscapeTextureWeightPatchPS, FGlobalShader);
+
+public:
+
+	// Flags that get packed into a bitfield because we're not allowed to use bool shader parameters:
+	enum class EFlags : uint8
+	{
+		None = 0,
+
+		// When false, falloff is circular.
+		RectangularFalloff = 1 << 0,
+
+		// When true, the texture alpha channel is considered for blending (in addition to falloff, if nonzero)
+		ApplyPatchAlpha = 1 << 1,
+	};
+
+	// TODO: We could consider exposing an additional global alpha setting that we can use to pass in the given
+	// edit layer alpha value... On the other hand, we currently don't bother doing this in any existing blueprint
+	// brushes, and it would be hard to support in a way that doesn't require each blueprint brush to respect it
+	// individually... Not clear whether this is something worth doing yet.
+
+	BEGIN_SHADER_PARAMETER_STRUCT(FParameters, )
+		SHADER_PARAMETER_RDG_TEXTURE_SRV(Texture2D<float4>, InSourceWeightmap)
+		// Offset of the source weightmap relative to the 0,0 location in the destination weightmap, because
+		// the source is likely to be a copied region from some inner part of the destination. This is basically
+		// a SourceWeightmapToDestinationWeightmap coordinate transformation, except that it is always a simple 
+		// integer translation.
+		SHADER_PARAMETER(FIntPoint, InSourceWeightmapCoordOffset)
+		SHADER_PARAMETER_RDG_TEXTURE_SRV(Texture2D<float4>, InWeightPatch)
+		SHADER_PARAMETER_SAMPLER(SamplerState, InWeightPatchSampler)
+		SHADER_PARAMETER(FMatrix44f, InWeightmapToPatch)
+		// Amount of the patch edge to not apply in UV space. Generally set to 0.5/Dimensions to avoid applying
+		// the edge half-pixels.
+		SHADER_PARAMETER(FVector2f, InEdgeUVDeadBorder)
+		// In world units, the size of the margin across which the alpha falls from 1 to 0
+		SHADER_PARAMETER(float, InFalloffWorldMargin)
+		// Size of the patch in world units (used for falloff)
+		SHADER_PARAMETER(FVector2f, InPatchWorldDimensions)
+		SHADER_PARAMETER(uint32, InBlendMode)
+		// Some combination of the flags (see constants above).
+		SHADER_PARAMETER(uint32, InFlags)
+
+		RENDER_TARGET_BINDING_SLOTS() // Holds our output
+	END_SHADER_PARAMETER_STRUCT()
+
+	static bool ShouldCompilePermutation(const FGlobalShaderPermutationParameters& Parameters);
+	static void ModifyCompilationEnvironment(const FGlobalShaderPermutationParameters& InParameters, FShaderCompilerEnvironment& OutEnvironment);
+
+	static void AddToRenderGraph(FRDGBuilder& GraphBuilder, FParameters* InParameters, const FIntRect& DestinationBounds);
+};
+
+// We need an alias for FApplyLandscapeTextureWeightPatchPSParameters::FParameters so that it can be forward-declared in another header file (cannot forward-declare internal classes)
+class FApplyLandscapeTextureWeightPatchPSParameters : public FApplyLandscapeTextureWeightPatchPS::FParameters {};
+
+ENUM_CLASS_FLAGS(FApplyLandscapeTextureWeightPatchPS::EFlags);
+
+class FReinitializeLandscapePatchPS : public FGlobalShader
+{
+	DECLARE_GLOBAL_SHADER(FReinitializeLandscapePatchPS);
+	SHADER_USE_PARAMETER_STRUCT(FReinitializeLandscapePatchPS, FGlobalShader);
+
+public:
+
+	class FHeightPatch : SHADER_PERMUTATION_BOOL("HEIGHT_PATCH");
+	using FPermutationDomain = TShaderPermutationDomain<FHeightPatch>;
+
+	BEGIN_SHADER_PARAMETER_STRUCT(FParameters, )
+		SHADER_PARAMETER_RDG_TEXTURE_SRV(Texture2D<float4>, InSource)
+		SHADER_PARAMETER_SAMPLER(SamplerState, InSourceSampler)
+		SHADER_PARAMETER(FMatrix44f, InPatchToSource)
+
+		RENDER_TARGET_BINDING_SLOTS() // Holds our output
+	END_SHADER_PARAMETER_STRUCT()
+
+	static bool ShouldCompilePermutation(const FGlobalShaderPermutationParameters& Parameters);
+	static void ModifyCompilationEnvironment(const FGlobalShaderPermutationParameters& InParameters, FShaderCompilerEnvironment& OutEnvironment);
+
+	static void AddToRenderGraph(FRDGBuilder& GraphBuilder, FParameters* InParameters, bool bHeightPatch);
+};
+
+/**
+ * Simple shader for copying textures of potentially different resolutions.
+ */
+//~ Theoretically CopyToResolveTarget or AddCopyToResolveTargetPass should work, but I was unable to make them
+//~ work without lots of complaints from the renderer.
+class FSimpleTextureCopyPS : public FGlobalShader
+{
+	DECLARE_GLOBAL_SHADER(FSimpleTextureCopyPS);
+	SHADER_USE_PARAMETER_STRUCT(FSimpleTextureCopyPS, FGlobalShader);
+
+public:
+
+	BEGIN_SHADER_PARAMETER_STRUCT(FParameters, )
+		SHADER_PARAMETER_RDG_TEXTURE_SRV(Texture2D<float4>, InSource)
+		SHADER_PARAMETER_SAMPLER(SamplerState, InSourceSampler)
+		SHADER_PARAMETER(FVector2f, InDestinationResolution)
+		RENDER_TARGET_BINDING_SLOTS() // Holds our output
+	END_SHADER_PARAMETER_STRUCT()
+
+	static bool ShouldCompilePermutation(const FGlobalShaderPermutationParameters& Parameters);
+	static void ModifyCompilationEnvironment(const FGlobalShaderPermutationParameters& InParameters, FShaderCompilerEnvironment& OutEnvironment);
+
+	static void AddToRenderGraph(FRDGBuilder& GraphBuilder, FRDGTextureRef SourceTexture, FRDGTextureRef DestinationTexture);
+};
+
+}//end UE::Landscape
diff --git a/Engine/Plugins/Experimental/LandscapePatch/Source/LandscapePatch/Public/LandscapeCircleHeightPatch.h b/Engine/Plugins/Experimental/LandscapePatch/Source/LandscapePatch/Public/LandscapeCircleHeightPatch.h
--- a/Engine/Plugins/Experimental/LandscapePatch/Source/LandscapePatch/Public/LandscapeCircleHeightPatch.h
+++ b/Engine/Plugins/Experimental/LandscapePatch/Source/LandscapePatch/Public/LandscapeCircleHeightPatch.h
@@ -13,7 +13,7 @@
  * alpha decreases linearly. When added to an actor, initializes itself to the bottom of the bounding box.
  */
 UCLASS(Blueprintable, BlueprintType, ClassGroup = Landscape, meta = (BlueprintSpawnableComponent))
-class LANDSCAPEPATCH_API ULandscapeCircleHeightPatch : public ULandscapePatchComponent
+class ULandscapeCircleHeightPatch : public ULandscapePatchComponent
 {
 	GENERATED_BODY()
 
@@ -64,7 +64,3 @@ private:
 		const TFunction<FRHITexture* ()>& RenderThreadLandscapeTextureGetter, int32 LandscapeTextureSliceIndex,
 		const FIntPoint& DestinationResolution, const FTransform& HeightmapCoordsToWorld);
 };
-
-#if UE_ENABLE_INCLUDE_ORDER_DEPRECATED_IN_5_2
-#include "CoreMinimal.h"
-#endif
diff --git a/Engine/Plugins/Experimental/LandscapePatch/Source/LandscapePatch/Public/LandscapeCircleHeightPatchPS.h b/Engine/Plugins/Experimental/LandscapePatch/Source/LandscapePatch/Public/LandscapeCircleHeightPatchPS.h
deleted file mode 100644
diff --git a/Engine/Plugins/Experimental/LandscapePatch/Source/LandscapePatch/Public/LandscapePatchComponent.h b/Engine/Plugins/Experimental/LandscapePatch/Source/LandscapePatch/Public/LandscapePatchComponent.h
--- a/Engine/Plugins/Experimental/LandscapePatch/Source/LandscapePatch/Public/LandscapePatchComponent.h
+++ b/Engine/Plugins/Experimental/LandscapePatch/Source/LandscapePatch/Public/LandscapePatchComponent.h
@@ -44,8 +44,8 @@ enum class ELandscapePatchPriorityInitialization : uint8
  */
 //~ TODO: Although this doesn't generate geometry, we are likely to change this to inherit from UPrimitiveComponent
 //~ so that we can use render proxies for passing along data to the render thread or perhaps for visualization.
-UCLASS(Blueprintable, BlueprintType, Abstract)
-class LANDSCAPEPATCH_API ULandscapePatchComponent : public USceneComponent
+UCLASS(MinimalAPI, Blueprintable, BlueprintType, Abstract)
+class ULandscapePatchComponent : public USceneComponent
 	, public ILandscapeEditLayerRenderer
 {
 	GENERATED_BODY()
@@ -311,7 +311,3 @@ struct FLandscapePatchComponentInstanceData : public FSceneComponentInstanceData
 	bool bGaveMissingLandscapeWarning = false;
 #endif
 };
-
-#if UE_ENABLE_INCLUDE_ORDER_DEPRECATED_IN_5_2
-#include "CoreMinimal.h"
-#endif
diff --git a/Engine/Plugins/Experimental/LandscapePatch/Source/LandscapePatch/Public/LandscapePatchLogging.h b/Engine/Plugins/Experimental/LandscapePatch/Source/LandscapePatch/Public/LandscapePatchLogging.h
--- a/Engine/Plugins/Experimental/LandscapePatch/Source/LandscapePatch/Public/LandscapePatchLogging.h
+++ b/Engine/Plugins/Experimental/LandscapePatch/Source/LandscapePatch/Public/LandscapePatchLogging.h
@@ -5,7 +5,3 @@
 #include "Logging/LogMacros.h"
 
 LANDSCAPEPATCH_API DECLARE_LOG_CATEGORY_EXTERN(LogLandscapePatch, Log, All);
-
-#if UE_ENABLE_INCLUDE_ORDER_DEPRECATED_IN_5_2
-#include "CoreMinimal.h"
-#endif
diff --git a/Engine/Plugins/Experimental/LandscapePatch/Source/LandscapePatch/Public/LandscapePatchManager.h b/Engine/Plugins/Experimental/LandscapePatch/Source/LandscapePatch/Public/LandscapePatchManager.h
--- a/Engine/Plugins/Experimental/LandscapePatch/Source/LandscapePatch/Public/LandscapePatchManager.h
+++ b/Engine/Plugins/Experimental/LandscapePatch/Source/LandscapePatch/Public/LandscapePatchManager.h
@@ -148,7 +148,3 @@ private:
 	UPROPERTY(Transient)
 	TMap<TSoftObjectPtr<ULandscapePatchComponent>, int32> PatchToIndex;
 };
-
-#if UE_ENABLE_INCLUDE_ORDER_DEPRECATED_IN_5_2
-#include "CoreMinimal.h"
-#endif
diff --git a/Engine/Plugins/Experimental/LandscapePatch/Source/LandscapePatch/Public/LandscapePatchModule.h b/Engine/Plugins/Experimental/LandscapePatch/Source/LandscapePatch/Public/LandscapePatchModule.h
--- a/Engine/Plugins/Experimental/LandscapePatch/Source/LandscapePatch/Public/LandscapePatchModule.h
+++ b/Engine/Plugins/Experimental/LandscapePatch/Source/LandscapePatch/Public/LandscapePatchModule.h
@@ -12,8 +12,3 @@ public:
 	virtual void StartupModule() override;
 	virtual void ShutdownModule() override;
 };
-
-#if UE_ENABLE_INCLUDE_ORDER_DEPRECATED_IN_5_2
-#include "CoreMinimal.h"
-#include "Modules/ModuleManager.h"
-#endif
diff --git a/Engine/Plugins/Experimental/LandscapePatch/Source/LandscapePatch/Public/LandscapePatchUtil.h b/Engine/Plugins/Experimental/LandscapePatch/Source/LandscapePatch/Public/LandscapePatchUtil.h
--- a/Engine/Plugins/Experimental/LandscapePatch/Source/LandscapePatch/Public/LandscapePatchUtil.h
+++ b/Engine/Plugins/Experimental/LandscapePatch/Source/LandscapePatch/Public/LandscapePatchUtil.h
@@ -4,9 +4,33 @@
 
 #include "Math/MathFwd.h" // FTransform
 
+#include "LandscapePatchUtil.generated.h"
+
 class FRHICommandListImmediate;
 class FTextureResource;
 
+
+// ----------------------------------------------------------------------------------
+
+// Values needed to convert a patch stored in some source encoding into the native (two byte int) encoding and back
+USTRUCT()
+struct FLandscapeHeightPatchConvertToNativeParams
+{
+	GENERATED_BODY()
+
+	UPROPERTY()
+	float ZeroInEncoding = 0.0f;
+
+	UPROPERTY()
+	float HeightScale = 1.0f;
+
+	UPROPERTY()
+	float HeightOffset = 0.0f;
+};
+
+
+// ----------------------------------------------------------------------------------
+
 namespace UE:: Landscape::PatchUtil 
 {
 	void CopyTextureOnRenderThread(FRHICommandListImmediate& RHICmdList, const FTextureResource& Source, FTextureResource& Destination);
diff --git a/Engine/Plugins/Experimental/LandscapePatch/Source/LandscapePatch/Public/LandscapeTextureBackedRenderTarget.h b/Engine/Plugins/Experimental/LandscapePatch/Source/LandscapePatch/Public/LandscapeTextureBackedRenderTarget.h
--- a/Engine/Plugins/Experimental/LandscapePatch/Source/LandscapePatch/Public/LandscapeTextureBackedRenderTarget.h
+++ b/Engine/Plugins/Experimental/LandscapePatch/Source/LandscapePatch/Public/LandscapeTextureBackedRenderTarget.h
@@ -4,7 +4,7 @@
 
 
 #include "Engine/TextureRenderTarget2D.h" // ETextureRenderTargetFormat
-#include "LandscapeTexturePatchPS.h" // FLandscapeHeightPatchConvertToNativeParams
+#include "LandscapePatchUtil.h" // FLandscapeHeightPatchConvertToNativeParams
 
 #include "LandscapeTextureBackedRenderTarget.generated.h"
 
@@ -119,7 +119,3 @@ protected:
 	UPROPERTY()
 	TEnumAsByte<ETextureRenderTargetFormat> RenderTargetFormat = ETextureRenderTargetFormat::RTF_R32f;
 };
-
-#if UE_ENABLE_INCLUDE_ORDER_DEPRECATED_IN_5_2
-#include "CoreMinimal.h"
-#endif
diff --git a/Engine/Plugins/Experimental/LandscapePatch/Source/LandscapePatch/Public/LandscapeTexturePatch.h b/Engine/Plugins/Experimental/LandscapePatch/Source/LandscapePatch/Public/LandscapeTexturePatch.h
--- a/Engine/Plugins/Experimental/LandscapePatch/Source/LandscapePatch/Public/LandscapeTexturePatch.h
+++ b/Engine/Plugins/Experimental/LandscapePatch/Source/LandscapePatch/Public/LandscapeTexturePatch.h
@@ -14,6 +14,12 @@ class ULandscapeHeightTextureBackedRenderTarget;
 class UTexture;
 class UTexture2D;
 
+namespace UE::Landscape
+{
+	class FApplyLandscapeTextureHeightPatchPSParameters;
+	class FApplyLandscapeTextureWeightPatchPSParameters;
+} // namespace UE::Landscape
+
 /**
  * Determines where the patch gets its information, which affects its memory usage in editor (not in runtime,
  * since patches are baked directly into landscape and removed for runtime).
@@ -54,18 +60,18 @@ enum class ELandscapeTexturePatchBlendMode : uint8
 	// height using falloff/alpha. E.g. with no falloff and alpha 1, the landscape will
 	// be set directly to the height sampled from patch. With alpha 0.5, landscape height 
 	// will be averaged evenly with patch height.
-	AlphaBlend = static_cast<uint8>(UE::Landscape::FApplyLandscapeTextureHeightPatchPS::EBlendMode::AlphaBlend),
+	AlphaBlend,
 
 	// Interpreting the landscape mid value as 0, use the texture patch as an offset to
 	// apply to the landscape. Falloff/alpha will just affect the degree to which the offset
 	// is applied (e.g. alpha of 0.5 will apply just half the offset).
-	Additive = static_cast<uint8>(UE::Landscape::FApplyLandscapeTextureHeightPatchPS::EBlendMode::Additive),
+	Additive,
 
 	// Like Alpha Blend mode, but limited to only lowering the existing landscape values
-	Min = static_cast<uint8>(UE::Landscape::FApplyLandscapeTextureHeightPatchPS::EBlendMode::Min),
+	Min,
 
 	// Like Alpha Blend mode, but limited to only raising the existing landscape values
-	Max = static_cast<uint8>(UE::Landscape::FApplyLandscapeTextureHeightPatchPS::EBlendMode::Max)
+	Max
 };
 
 // Determines falloff method for the patch's influence.
@@ -121,7 +127,7 @@ enum class ELandscapeTextureHeightPatchZeroHeightMeaning : uint8
 
 //~ A struct in case we find that we need other encoding settings.
 USTRUCT(BlueprintType)
-struct LANDSCAPEPATCH_API FLandscapeTexturePatchEncodingSettings
+struct FLandscapeTexturePatchEncodingSettings
 {
 	GENERATED_BODY()
 public:
@@ -147,8 +153,8 @@ public:
  * Helper class for ULandscapeTexturePatch that stores information for a given weight layer.
  * Should not be used outside this class.
  */
-UCLASS(EditInlineNew, CollapseCategories)
-class LANDSCAPEPATCH_API ULandscapeWeightPatchTextureInfo : public UObject
+UCLASS(MinimalAPI, EditInlineNew, CollapseCategories)
+class ULandscapeWeightPatchTextureInfo : public UObject
 {
 	GENERATED_BODY()
 
@@ -212,8 +218,8 @@ protected:
 };
 
 
-UCLASS(Blueprintable, BlueprintType, ClassGroup = Landscape, meta = (BlueprintSpawnableComponent))
-class LANDSCAPEPATCH_API ULandscapeTexturePatch : public ULandscapePatchComponent
+UCLASS(MinimalAPI, Blueprintable, BlueprintType, ClassGroup = Landscape, meta = (BlueprintSpawnableComponent))
+class ULandscapeTexturePatch : public ULandscapePatchComponent
 {
 	GENERATED_BODY()
 
@@ -658,10 +664,10 @@ private:
 		FIntRect& DestinationBoundsOut, FVector2f& EdgeUVDeadBorderOut, float& FalloffWorldMarginOut) const;
 	void GetHeightShaderParams(const FTransform& LandscapeHeightmapToWorldIn, 
 		const FIntPoint& SourceResolutionIn, const FIntPoint& DestinationResolutionIn,
-		UE::Landscape::FApplyLandscapeTextureHeightPatchPS::FParameters& ParamsOut, FIntRect& DestinationBoundsOut) const;
+		UE::Landscape::FApplyLandscapeTextureHeightPatchPSParameters& ParamsOut, FIntRect& DestinationBoundsOut) const;
 	void GetWeightShaderParams(const FTransform& LandscapeHeightmapToWorldIn, const FIntPoint& SourceResolutionIn, 
 		const FIntPoint& DestinationResolutionIn, const ULandscapeWeightPatchTextureInfo* WeightPatchInfo, 
-		UE::Landscape::FApplyLandscapeTextureWeightPatchPS::FParameters& ParamsOut, FIntRect& DestinationBoundsOut) const;
+		UE::Landscape::FApplyLandscapeTextureWeightPatchPSParameters& ParamsOut, FIntRect& DestinationBoundsOut) const;
 	FMatrix44f GetPatchToHeightmapUVs(const FTransform& LandscapeHeightmapToWorld, int32 PatchSizeX, int32 PatchSizeY, int32 HeightmapSizeX, int32 HeightmapSizeY) const;
 	void ReinitializeHeight(UTextureRenderTarget2D* InCombinedResult, const FTransform& LandscapeHeightmapToWorld);
 	/**
@@ -679,7 +685,3 @@ private:
 	UPROPERTY(EditDefaultsOnly, Category = Settings)
 	TEnumAsByte<ETextureRenderTargetFormat> HeightRenderTargetFormat = ETextureRenderTargetFormat::RTF_R32f;
 };
-
-#if UE_ENABLE_INCLUDE_ORDER_DEPRECATED_IN_5_2
-#include "CoreMinimal.h"
-#endif
diff --git a/Engine/Plugins/Experimental/LandscapePatch/Source/LandscapePatch/Public/LandscapeTexturePatchPS.h b/Engine/Plugins/Experimental/LandscapePatch/Source/LandscapePatch/Public/LandscapeTexturePatchPS.h
deleted file mode 100644
diff --git a/Engine/Plugins/Experimental/LandscapePatch/Source/LandscapePatchEditorOnly/Private/LandscapeTexturePatchVisualizer.h b/Engine/Plugins/Experimental/LandscapePatch/Source/LandscapePatchEditorOnly/Private/LandscapeTexturePatchVisualizer.h
--- a/Engine/Plugins/Experimental/LandscapePatch/Source/LandscapePatchEditorOnly/Private/LandscapeTexturePatchVisualizer.h
+++ b/Engine/Plugins/Experimental/LandscapePatch/Source/LandscapePatchEditorOnly/Private/LandscapeTexturePatchVisualizer.h
@@ -13,7 +13,3 @@ class LANDSCAPEPATCHEDITORONLY_API FLandscapeTexturePatchVisualizer : public FCo
 private:
 	virtual void DrawVisualization(const UActorComponent* Component, const FSceneView* View, FPrimitiveDrawInterface* PDI) override;
 };
-
-#if UE_ENABLE_INCLUDE_ORDER_DEPRECATED_IN_5_2
-#include "CoreMinimal.h"
-#endif
diff --git a/Engine/Plugins/Experimental/LandscapePatch/Source/LandscapePatchEditorOnly/Public/LandscapePatchEditorOnlyModule.h b/Engine/Plugins/Experimental/LandscapePatch/Source/LandscapePatchEditorOnly/Public/LandscapePatchEditorOnlyModule.h
--- a/Engine/Plugins/Experimental/LandscapePatch/Source/LandscapePatchEditorOnly/Public/LandscapePatchEditorOnlyModule.h
+++ b/Engine/Plugins/Experimental/LandscapePatch/Source/LandscapePatchEditorOnly/Public/LandscapePatchEditorOnlyModule.h
@@ -18,8 +18,3 @@ private:
 	TArray<FName> ClassesToUnregisterOnShutdown;
 	TArray<FName> VisualizersToUnregisterOnShutdown;
 };
-
-#if UE_ENABLE_INCLUDE_ORDER_DEPRECATED_IN_5_2
-#include "CoreMinimal.h"
-#include "Modules/ModuleManager.h"
-#endif
