Non-Unity build fix

#rnx

diff --git a/Engine/Plugins/Experimental/LandscapePatch/Source/LandscapePatch/Private/LandscapeTextureBackedRenderTarget.cpp b/Engine/Plugins/Experimental/LandscapePatch/Source/LandscapePatch/Private/LandscapeTextureBackedRenderTarget.cpp
--- a/Engine/Plugins/Experimental/LandscapePatch/Source/LandscapePatch/Private/LandscapeTextureBackedRenderTarget.cpp
+++ b/Engine/Plugins/Experimental/LandscapePatch/Source/LandscapePatch/Private/LandscapeTextureBackedRenderTarget.cpp
@@ -6,6 +6,7 @@
 #include "LandscapeDataAccess.h"
 #include "LandscapePatchUtil.h" // CopyTextureOnRenderThread
 #include "LandscapePatchLogging.h" 
+#include "LandscapeTexturePatchPS.h"
 #include "RenderGraphBuilder.h"
 #include "TextureCompiler.h"
 #include "RenderingThread.h"
