LandscapePatch: Fix external texture asset patches slightly aligning incorrectly due to using incorrect resolution for coverage calculation. Existing instances that used an external texture asset should be checked and resaved. LandscapePatch.UseExternalTextureAlignmentFix can be set to 0 to revert to old behavior temporarily if needed.

diff --git a/Engine/Plugins/Experimental/LandscapePatch/Source/LandscapePatch/Private/LandscapeTexturePatch.cpp b/Engine/Plugins/Experimental/LandscapePatch/Source/LandscapePatch/Private/LandscapeTexturePatch.cpp
--- a/Engine/Plugins/Experimental/LandscapePatch/Source/LandscapePatch/Private/LandscapeTexturePatch.cpp
+++ b/Engine/Plugins/Experimental/LandscapePatch/Source/LandscapePatch/Private/LandscapeTexturePatch.cpp
@@ -12,7 +12,9 @@
 #include "LandscapePatchUtil.h" // CopyTextureOnRenderThread
 #include "LandscapeTexturePatchPS.h"
 #include "LandscapeUtils.h" // IsVisibilityLayer
+#include "Logging/MessageLog.h"
 #include "MathUtil.h"
+#include "Misc/UObjectToken.h"
 #include "RHIStaticStates.h"
 #include "RenderGraphUtils.h"
 #include "TextureResource.h"
@@ -21,6 +23,8 @@
 
 #include UE_INLINE_GENERATED_CPP_BY_NAME(LandscapeTexturePatch)
 
+#define LOCTEXT_NAMESPACE "LandscapeTexturePatch"
+
 namespace LandscapeTexturePatchLocals
 {
 #if WITH_EDITOR
@@ -104,6 +108,19 @@ namespace LandscapeTexturePatchLocals
 		}
 	}
 #endif // WITH_EDITOR
+
+	bool bUseExternalTextureFix = true;
+	FAutoConsoleVariableRef CVarUseExternalTextureAlignmentFix(
+		TEXT("LandscapePatch.UseExternalTextureAlignmentFix"), bUseExternalTextureFix, 
+		TEXT("Can be used to roll back an external texture alignment fix if temporarily needed."));
+
+	FVector2D GetCoverageWithExtraPixel(const FVector2D& Resolution, const FVector2D& Coverage)
+	{
+		// UnscaledPatchCoverage is meant to represent the distance between the centers of the extremal pixels.
+		//  That distance in pixels is Resolution-1.
+		FVector2D TargetPixelSize(Coverage / FVector2D::Max(Resolution - 1, FVector2D(1, 1)));
+		return TargetPixelSize * Resolution;
+	}
 }
 
 #if WITH_EDITOR
@@ -501,9 +518,12 @@ void ULandscapeTexturePatch::GetCommonShaderParams(const FTransform& LandscapeHe
 	FTransform& PatchToWorldOut, FVector2f& PatchWorldDimensionsOut, FMatrix44f& HeightmapToPatchOut, FIntRect& DestinationBoundsOut, 
 	FVector2f& EdgeUVDeadBorderOut, float& FalloffWorldMarginOut) const
 {
+	using namespace LandscapeTexturePatchLocals;
+
 	PatchToWorldOut = GetPatchToWorldTransform();
 
-	FVector2D FullPatchDimensions = GetFullUnscaledWorldSize();
+	FVector2D FullPatchDimensions = bUseExternalTextureFix ? GetCoverageWithExtraPixel(FVector2D(SourceResolutionIn), GetUnscaledCoverage())
+		: GetFullUnscaledWorldSize();
 	PatchWorldDimensionsOut = FVector2f(FullPatchDimensions);
 
 	FTransform FromPatchUVToPatch(FQuat4d::Identity, FVector3d(-FullPatchDimensions.X / 2, -FullPatchDimensions.Y / 2, 0),
@@ -1040,7 +1060,10 @@ void ULandscapeTexturePatch::ReinitializeWeightPatch(ULandscapeWeightPatchTextur
 FMatrix44f ULandscapeTexturePatch::GetPatchToHeightmapUVs(const FTransform& LandscapeHeightmapToWorld,
 	int32 PatchSizeX, int32 PatchSizeY, int32 HeightmapSizeX, int32 HeightmapSizeY) const
 {
-	FVector2D FullPatchDimensions = GetFullUnscaledWorldSize();
+	using namespace LandscapeTexturePatchLocals;
+
+	FVector2D FullPatchDimensions = bUseExternalTextureFix ? GetCoverageWithExtraPixel(FVector2D(PatchSizeX, PatchSizeY), GetUnscaledCoverage())
+		: GetFullUnscaledWorldSize();
 
 	FTransform PatchPixelToPatchLocal(FQuat4d::Identity, FVector3d(-FullPatchDimensions.X / 2, -FullPatchDimensions.Y / 2, 0),
 		FVector3d(FullPatchDimensions.X / PatchSizeX, FullPatchDimensions.Y / PatchSizeY, 1));
@@ -1171,6 +1194,47 @@ TStructOnScope<FActorComponentInstanceData> ULandscapeTexturePatch::GetComponent
 	return ToReturn;
 }
 
+void ULandscapeTexturePatch::CheckForErrors()
+{
+	auto DidPatchUseATextureAsset = [this]()
+	{
+		if (HeightSourceMode == ELandscapeTexturePatchSourceMode::TextureAsset && HeightTextureAsset)
+		{
+			return true;
+		}
+		for (const TObjectPtr<ULandscapeWeightPatchTextureInfo>& WeightPatch : WeightPatches)
+		{
+			if (IsValid(WeightPatch) && WeightPatch->SourceMode == ELandscapeTexturePatchSourceMode::TextureAsset
+				&& WeightPatch->TextureAsset)
+			{
+				return true;
+			}
+		}
+		return false;
+	};
+
+	if (!IsTemplate()
+		&& GetLinkerCustomVersion(FFortniteMainBranchObjectVersion::GUID) < FFortniteMainBranchObjectVersion::LandscapeTexturePatchUsesTextureAssetResolution
+		&& DidPatchUseATextureAsset())
+	{
+		FFormatNamedArguments Arguments;
+		Arguments.Add(TEXT("Package"), FText::FromString(*GetNameSafe(GetPackage())));
+		Arguments.Add(TEXT("Actor"), FText::FromString(*GetNameSafe(this)));
+
+		FMessageLog("MapCheck").Warning()
+			->AddToken(FUObjectToken::Create(this))
+			->AddToken(FTextToken::Create(FText::Format(LOCTEXT("PatchAlignmentMightBeAdjusted", "Patch was saved with a texture "
+				"asset when alignment code for texture assets had a minor bug. Verify that difference in landscape is negligible, and "
+				"resave the patch. "
+				"(Package: {Package}, Actor: {Actor})."), Arguments)))
+			->AddToken(FActionToken::Create(LOCTEXT("MarkDirtyButton", "Mark dirty"), FText(),
+				FOnActionTokenExecuted::CreateWeakLambda(this, [this]()
+					{
+						MarkPackageDirty();
+					})));
+	}
+}
+
 #endif
 
 void ULandscapeTexturePatch::SnapToLandscape()
@@ -1311,12 +1375,9 @@ void ULandscapeTexturePatch::SetResolution(FVector2D ResolutionIn)
 
 FVector2D ULandscapeTexturePatch::GetFullUnscaledWorldSize() const
 {
-	FVector2D Resolution = GetResolution();
+	using namespace LandscapeTexturePatchLocals;
 
-	// UnscaledPatchCoverage is meant to represent the distance between the centers of the extremal pixels.
-	// That distance in pixels is Resolution-1.
-	FVector2D TargetPixelSize(UnscaledPatchCoverage / FVector2D::Max(Resolution - 1, FVector2D(1, 1)));
-	return TargetPixelSize * Resolution;
+	return GetCoverageWithExtraPixel(GetResolution(), UnscaledPatchCoverage);
 }
 
 FTransform ULandscapeTexturePatch::GetPatchToWorldTransform() const
@@ -1415,6 +1476,12 @@ void ULandscapeTexturePatch::PostEditChangeProperty(FPropertyChangedEvent& Prope
 	Super::PostEditChangeProperty(PropertyChangedEvent);
 }
 
+void ULandscapeTexturePatch::Serialize(FArchive& Ar)
+{
+	Super::Serialize(Ar);
+	Ar.UsingCustomVersion(FFortniteMainBranchObjectVersion::GUID);
+}
+
 void ULandscapeWeightPatchTextureInfo::PostEditChangeProperty(FPropertyChangedEvent& PropertyChangedEvent)
 {
 	using namespace LandscapeTexturePatchLocals;
@@ -1997,7 +2064,7 @@ TArray<UE::Landscape::EditLayers::FEditLayerRenderItem> ULandscapeTexturePatch::
 	TArray<FEditLayerRenderItem> AffectedAreas;
 
 	FTransform ComponentTransform = this->GetComponentToWorld();
-	FOOBox2D PatchArea(ComponentTransform, GetFullUnscaledWorldSize());
+	FOOBox2D PatchArea(ComponentTransform, GetUnscaledCoverage());
 	FInputWorldArea InputWorldArea = FInputWorldArea::CreateOOBox(PatchArea);
 	FOutputWorldArea OutputWorldArea = FOutputWorldArea::CreateOOBox(PatchArea);
 
@@ -2023,3 +2090,5 @@ TArray<UE::Landscape::EditLayers::FEditLayerRenderItem> ULandscapeTexturePatch::
 	return AffectedAreas;
 }
 #endif // WITH_EDITOR
+
+#undef LOCTEXT_NAMESPACE
diff --git a/Engine/Plugins/Experimental/LandscapePatch/Source/LandscapePatch/Public/LandscapeTexturePatch.h b/Engine/Plugins/Experimental/LandscapePatch/Source/LandscapePatch/Public/LandscapeTexturePatch.h
--- a/Engine/Plugins/Experimental/LandscapePatch/Source/LandscapePatch/Public/LandscapeTexturePatch.h
+++ b/Engine/Plugins/Experimental/LandscapePatch/Source/LandscapePatch/Public/LandscapeTexturePatch.h
@@ -248,10 +248,12 @@ public:
 
 	// UActorComponent
 	virtual TStructOnScope<FActorComponentInstanceData> GetComponentInstanceData() const override;
+	virtual void CheckForErrors() override;
 
 	// UObject
 	virtual void PostEditChangeProperty(FPropertyChangedEvent& PropertyChangedEvent) override;
 #endif
+	virtual void Serialize(FArchive& Ar) override;
 
 	/**
 	 * Gets the transform from patch to world. The transform is based off of the component
@@ -263,10 +265,9 @@ public:
 
 	/**
 	 * Gives size in unscaled world coordinates (ie before applying patch transform) of the patch as measured 
-	 * between the centers of the outermost pixels. Measuring the coverage this way means that a patch can 
-	 * affect the same region of the landscape regardless of patch resolution.
-	 * This is also the range across which bilinear interpolation always has correct values, so the area outside 
-	 * this center portion is usually set as a "dead" border that doesn't affect the landscape.
+	 * between the centers of the outermost pixels. This is the range across which bilinear interpolation
+	 * always has correct values, so the area outside this center portion in the texture does not affect
+	 * the landscape.
 	 */
 	UFUNCTION(BlueprintCallable, Category = LandscapePatch)
 	virtual FVector2D GetUnscaledCoverage() const { return FVector2D(UnscaledPatchCoverage); }
@@ -278,8 +279,10 @@ public:
 	virtual void SetUnscaledCoverage(FVector2D Coverage) { UnscaledPatchCoverage = Coverage; }
 
 	/**
-	 * Gives size in unscaled world coordinates of the patch in the world, based off of UnscaledCoverage and
-	 * texture resolution (i.e., adds a half-pixel around UnscaledCoverage).
+	 * When using an internal texture, gives size in unscaled world coordinates of the patch in the world,
+	 * based off of UnscaledCoverage and texture resolution (i.e., adds a half-pixel around UnscaledCoverage).
+	 * Does not reflect the resolution of any used texture assets (if the source mode is texture asset for
+	 * the height/weight patches).
 	 */
 	UFUNCTION(BlueprintCallable, Category = LandscapePatch)
 	virtual FVector2D GetFullUnscaledWorldSize() const;
diff --git a/Engine/Source/Runtime/Core/Public/UObject/FortniteMainBranchObjectVersion.h b/Engine/Source/Runtime/Core/Public/UObject/FortniteMainBranchObjectVersion.h
--- a/Engine/Source/Runtime/Core/Public/UObject/FortniteMainBranchObjectVersion.h
+++ b/Engine/Source/Runtime/Core/Public/UObject/FortniteMainBranchObjectVersion.h
@@ -541,6 +541,9 @@ struct FFortniteMainBranchObjectVersion
 		// Added managed array property serialization
 		AddManagedArrayCollectionPropertySerialization,
 
+		// Landscape texture patches in Texture Asset source mode now use proper resolution when calculating transform
+		LandscapeTexturePatchUsesTextureAssetResolution,
+
 		// -----<new versions can be added above this line>-------------------------------------------------
 		VersionPlusOne,
 		LatestVersion = VersionPlusOne - 1
