LandscapePatch: Minor fix- make priority initialization not set incorrect priority when there isn't a landscape in the level.

diff --git a/Engine/Plugins/Experimental/LandscapePatch/Source/LandscapePatch/Private/LandscapePatchComponent.cpp b/Engine/Plugins/Experimental/LandscapePatch/Source/LandscapePatch/Private/LandscapePatchComponent.cpp
--- a/Engine/Plugins/Experimental/LandscapePatch/Source/LandscapePatch/Private/LandscapePatchComponent.cpp
+++ b/Engine/Plugins/Experimental/LandscapePatch/Source/LandscapePatch/Private/LandscapePatchComponent.cpp
@@ -326,7 +326,8 @@ void ULandscapePatchComponent::OnComponentCreated()
 	// We're going to be binding to an edit layer, which will place us in its registered patch list with
 	//  our current priority. If we're later going to be updating our priority to be the highest, then 
 	//  we need to temporarily lower our priority so that we aren't accidentally the highest priority patch
-	//  the edit layer sees when we query it. 
+	//  the edit layer sees when we query it.
+	double PriorityToResetTo = Priority; // may be needed to undo the following
 	if (IsPatchInWorld() && PriorityInitialization == ELandscapePatchPriorityInitialization::AcquireHighest)
 	{
 		Priority = TNumericLimits<double>::Lowest();
@@ -372,6 +373,12 @@ void ULandscapePatchComponent::OnComponentCreated()
 				Priority = EditLayer->GetHighestPatchPriority() + 1;
 				EditLayer->NotifyOfPriorityChange(this);
 			}
+			else
+			{
+				// If we failed to get an edit layer to query, undo the setting of low priority
+				//  we did above.
+				Priority = PriorityToResetTo;
+			}
 			break;
 		}
 	}
