Expose ReadOnly blueprint access for  HeightTextureAsset and HeightEncodingSettings in  LandscapePatchComponent

diff --git a/Engine/Plugins/Experimental/LandscapePatch/Source/LandscapePatch/Private/LandscapeTexturePatch.cpp b/Engine/Plugins/Experimental/LandscapePatch/Source/LandscapePatch/Private/LandscapeTexturePatch.cpp
--- a/Engine/Plugins/Experimental/LandscapePatch/Source/LandscapePatch/Private/LandscapeTexturePatch.cpp
+++ b/Engine/Plugins/Experimental/LandscapePatch/Source/LandscapePatch/Private/LandscapeTexturePatch.cpp
@@ -1887,6 +1887,18 @@ ELandscapeTexturePatchSourceMode ULandscapeTexturePatch::GetWeightPatchSourceMod
 	return ELandscapeTexturePatchSourceMode::None;
 }
 
+UTexture* ULandscapeTexturePatch::GetWeightPatchTextureAsset(const FName& InWeightmapLayerName)
+{
+	for (const TObjectPtr<ULandscapeWeightPatchTextureInfo>& WeightPatch : WeightPatches)
+	{
+		if (IsValid(WeightPatch) && WeightPatch->WeightmapLayerName == InWeightmapLayerName)
+		{
+			return WeightPatch->TextureAsset;
+		}
+	}
+	return nullptr;
+}
+
 UTextureRenderTarget2D* ULandscapeTexturePatch::GetWeightPatchRenderTarget(const FName& InWeightmapLayerName, bool bMarkDirty)
 {
 	if (IsTemplate())
diff --git a/Engine/Plugins/Experimental/LandscapePatch/Source/LandscapePatch/Public/LandscapeTexturePatch.h b/Engine/Plugins/Experimental/LandscapePatch/Source/LandscapePatch/Public/LandscapeTexturePatch.h
--- a/Engine/Plugins/Experimental/LandscapePatch/Source/LandscapePatch/Public/LandscapeTexturePatch.h
+++ b/Engine/Plugins/Experimental/LandscapePatch/Source/LandscapePatch/Public/LandscapeTexturePatch.h
@@ -465,7 +465,7 @@ public:
 	UFUNCTION(BlueprintCallable, Category = LandscapePatch)
 	virtual void RemoveAllWeightPatches();
 
-	/** Sets the soure mode of all weight patches to "None". */
+	/** Sets the source mode of all weight patches to "None". */
 	UFUNCTION(BlueprintCallable, Category = LandscapePatch)
 	virtual void DisableAllWeightPatches();
 
@@ -475,6 +475,9 @@ public:
 	UFUNCTION(BlueprintCallable, Category = LandscapePatch)
 	virtual ELandscapeTexturePatchSourceMode GetWeightPatchSourceMode(const FName& InWeightmapLayerName);
 
+	UFUNCTION(BlueprintCallable, Category = LandscapePatch)
+	virtual UTexture* GetWeightPatchTextureAsset(const FName& InWeightmapLayerName);
+
 	UFUNCTION(BlueprintCallable, Category = LandscapePatch)
 	virtual void SetWeightPatchSourceMode(const FName& InWeightmapLayerName, ELandscapeTexturePatchSourceMode NewMode);
 
@@ -550,7 +553,7 @@ protected:
 	/**
 	 * Texture used when source mode is set to a texture asset.
 	 */
-	UPROPERTY(EditAnywhere, Category = HeightPatch, meta = (EditConditionHides,
+	UPROPERTY(EditAnywhere, BlueprintReadOnly, Category = HeightPatch, meta = (EditConditionHides,
 		EditCondition = "HeightSourceMode == ELandscapeTexturePatchSourceMode::TextureAsset", 
 		DisallowedAssetDataTags = "VirtualTextureStreaming=True"))
 	TObjectPtr<UTexture> HeightTextureAsset = nullptr;
@@ -569,7 +572,7 @@ protected:
 	ELandscapeTextureHeightPatchEncoding HeightEncoding = ELandscapeTextureHeightPatchEncoding::WorldUnits;
 
 	/** Encoding settings. Not relevant when using native packed height as the encoding. */
-	UPROPERTY(EditAnywhere, Category = HeightPatch, meta = (UIMin = "0", UIMax = "1",
+	UPROPERTY(EditAnywhere, BlueprintReadOnly, Category = HeightPatch, meta = (UIMin = "0", UIMax = "1",
 		EditCondition = "HeightSourceMode != ELandscapeTexturePatchSourceMode::InternalTexture && HeightEncoding != ELandscapeTextureHeightPatchEncoding::NativePackedHeight"))
 	FLandscapeTexturePatchEncodingSettings HeightEncodingSettings;
 
