Fixed memory stomp when "reinitializing" heightmap or weightmap in landscape patches : if the backing UTexture was just created, it might copy data from the default texture resource because the texture was not compiled yet

diff --git a/Engine/Plugins/Experimental/LandscapePatch/Source/LandscapePatch/Private/LandscapeTextureBackedRenderTarget.cpp b/Engine/Plugins/Experimental/LandscapePatch/Source/LandscapePatch/Private/LandscapeTextureBackedRenderTarget.cpp
--- a/Engine/Plugins/Experimental/LandscapePatch/Source/LandscapePatch/Private/LandscapeTextureBackedRenderTarget.cpp
+++ b/Engine/Plugins/Experimental/LandscapePatch/Source/LandscapePatch/Private/LandscapeTextureBackedRenderTarget.cpp
@@ -41,15 +41,20 @@ namespace LandscapeTextureBackedRenderTargetLocals
 	{
 		if (Texture)
 		{
-			TArray<FColor> Pixels;
-			Pixels.Init(ClearColor, Texture->GetSizeX() * Texture->GetSizeY());
-
-			void* TextureData = Texture->GetPlatformData()->Mips[0].BulkData.Lock(LOCK_READ_WRITE);
-
-			FMemory::Memcpy(TextureData, Pixels.GetData(), Pixels.Num() * sizeof(FColor));
-
-			Texture->GetPlatformData()->Mips[0].BulkData.Unlock();
-			Texture->UpdateResource();
+			// We need to wait for the platform data to be ready (if the texture was just created, it likely won't) :
+			FTextureCompilingManager::Get().FinishCompilation({ Texture });
+			if (ensure((Texture->GetPlatformData()->SizeX == Texture->GetSizeX()) && (Texture->GetPlatformData()->SizeY == Texture->GetSizeY())))
+			{
+			    TArray<FColor> Pixels;
+			    Pixels.Init(ClearColor, Texture->GetSizeX() * Texture->GetSizeY());
+    
+			    void* TextureData = Texture->GetPlatformData()->Mips[0].BulkData.Lock(LOCK_READ_WRITE);
+    
+			    FMemory::Memcpy(TextureData, Pixels.GetData(), Pixels.Num() * sizeof(FColor));
+    
+			    Texture->GetPlatformData()->Mips[0].BulkData.Unlock();
+			    Texture->UpdateResource();
+			}
 		}
 	}
 #endif // WITH_EDITOR
