Fixed assert with landscape patches : scratch textures can have more slices than actually necessary (because they come from a pool)

#rnx

diff --git a/Engine/Plugins/Experimental/LandscapePatch/Source/LandscapePatch/Private/LandscapeCircleHeightPatch.cpp b/Engine/Plugins/Experimental/LandscapePatch/Source/LandscapePatch/Private/LandscapeCircleHeightPatch.cpp
--- a/Engine/Plugins/Experimental/LandscapePatch/Source/LandscapePatch/Private/LandscapeCircleHeightPatch.cpp
+++ b/Engine/Plugins/Experimental/LandscapePatch/Source/LandscapePatch/Private/LandscapeCircleHeightPatch.cpp
@@ -97,7 +97,7 @@ void ULandscapeCircleHeightPatch::RenderLayer(FRenderParams& InRenderParams, UE:
 	check(TextureArray != nullptr);
 
 	int32 NumLayers = InRenderParams.TargetLayerGroupLayerInfos.Num();
-	check(TextureArray->Slices == NumLayers);
+	check(LandscapeRT->GetEffectiveNumSlices() == NumLayers);
 
 	for (int32 LayerIndex = 0; LayerIndex < NumLayers; ++LayerIndex)
 	{
diff --git a/Engine/Plugins/Experimental/LandscapePatch/Source/LandscapePatch/Private/LandscapeTexturePatch.cpp b/Engine/Plugins/Experimental/LandscapePatch/Source/LandscapePatch/Private/LandscapeTexturePatch.cpp
--- a/Engine/Plugins/Experimental/LandscapePatch/Source/LandscapePatch/Private/LandscapeTexturePatch.cpp
+++ b/Engine/Plugins/Experimental/LandscapePatch/Source/LandscapePatch/Private/LandscapeTexturePatch.cpp
@@ -183,7 +183,7 @@ void ULandscapeTexturePatch::RenderLayer(FRenderParams& InRenderParams, UE::Land
 	check(TextureArray != nullptr);
 	
 	int32 NumTargetLayersInGroup = InRenderParams.TargetLayerGroupLayerNames.Num();
-	check(TextureArray->Slices == NumTargetLayersInGroup);
+	check(LandscapeScratchRT->GetEffectiveNumSlices() == NumTargetLayersInGroup);
 
 	for (int32 TargetLayerIndexInGroup = 0; TargetLayerIndexInGroup < NumTargetLayersInGroup; ++TargetLayerIndexInGroup)
 	{
