Automatically migrate legacy LandscapePatchManager setups to the new system when a scene is loaded.
Can be disabled by setting "LandscapePatch.AutoMigrateLegacyListToPrioritySystem=0".

diff --git a/Engine/Plugins/Experimental/LandscapePatch/Source/LandscapePatch/Private/LandscapePatchComponent.cpp b/Engine/Plugins/Experimental/LandscapePatch/Source/LandscapePatch/Private/LandscapePatchComponent.cpp
--- a/Engine/Plugins/Experimental/LandscapePatch/Source/LandscapePatch/Private/LandscapePatchComponent.cpp
+++ b/Engine/Plugins/Experimental/LandscapePatch/Source/LandscapePatch/Private/LandscapePatchComponent.cpp
@@ -149,6 +149,7 @@ void ULandscapePatchComponent::CheckForErrors()
 
 	if (PatchManager.IsValid())
 	{
+		ensure(!PatchManager->IsDead());
 		if (EditLayerGuid.IsValid())
 		{
 			FMessageLog("MapCheck").Warning()
@@ -307,6 +308,7 @@ void ULandscapePatchComponent::OnComponentCreated()
 	// make sure our other pointers are consistent.
 	if (PatchManager.IsValid())
 	{
+		check(!PatchManager->IsDead());
 		if (Landscape != PatchManager->GetOwningLandscape())
 		{
 			Modify();
@@ -851,6 +853,11 @@ bool ULandscapePatchComponent::BindToEditLayer(const FGuid& Guid)
 		}
 	}
 
+	if (EditLayerGuid != Guid)
+	{
+		Modify();  // Changing the guid, so mark the package dirty.
+	}
+
 	EditLayer = Cast<ULandscapePatchEditLayer>(LayerStruct ? LayerStruct->EditLayer : nullptr);
 	EditLayerGuid = Guid;
 	DetailPanelLayerGuid = Guid.ToString();
@@ -1015,7 +1022,7 @@ void ULandscapePatchComponent::FixBindings()
 
 	if (PatchManager.IsValid())
 	{
-		if (PatchManager->GetOwningLandscape())
+		if (!PatchManager->IsDead() && PatchManager->GetOwningLandscape())
 		{
 			ResetEditLayer();
 			Landscape = PatchManager->GetOwningLandscape();
diff --git a/Engine/Plugins/Experimental/LandscapePatch/Source/LandscapePatch/Private/LandscapePatchManager.cpp b/Engine/Plugins/Experimental/LandscapePatch/Source/LandscapePatch/Private/LandscapePatchManager.cpp
--- a/Engine/Plugins/Experimental/LandscapePatch/Source/LandscapePatch/Private/LandscapePatchManager.cpp
+++ b/Engine/Plugins/Experimental/LandscapePatch/Source/LandscapePatch/Private/LandscapePatchManager.cpp
@@ -40,6 +40,11 @@ namespace LandscapePatchManagerLocals
 {
 	const FText MigratePatchesTransactionName(LOCTEXT("MigratePatchesTransaction", "Migrate Patches"));
 
+	static FAutoConsoleVariable CVarMigrateLegacyPatchListToPrioritySystem(
+		TEXT("LandscapePatch.AutoMigrateLegacyListToPrioritySystem"),
+		true,
+		TEXT("When loaded in the editor, automatically remove all LandscapePatchManagers and bind their patches directly to an edit layer.  Set their patch priorities according to their index."));
+
 #if WITH_EDITOR
 	// Note: the priorities would become jumbled up if someone had multiple managers in the same edit layer. But
 	// this is an unexpected case that is not worth trying to handle differently.
@@ -181,6 +186,33 @@ namespace LandscapePatchManagerLocals
 			return false;
 		});
 	}
+
+#if WITH_EDITOR
+	// Try to get a soft pointer to the owner actor for a ULandscapePatchComponent when that actor and component have not loaded.
+	TSoftObjectPtr<UObject> GetActorPtrFromPatchComponent(TSoftObjectPtr<ULandscapePatchComponent> Patch)
+	{
+		const FSoftObjectPath& PatchPath = Patch.ToSoftObjectPath();
+		if (PatchPath.IsSubobject())
+		{
+			FTopLevelAssetPath AssetPath = PatchPath.GetAssetPath();
+			FUtf8String SubPath = PatchPath.GetSubPathUtf8String();
+
+			// Remove the last part of the subpath, which should be the component name ".LandscapeTexturePatch", etc.
+			int32 LastPathSeparator = INDEX_NONE;
+			if (SubPath.FindLastChar(UTF8TEXT('.'), LastPathSeparator))
+			{
+				check(LastPathSeparator != INDEX_NONE);
+				SubPath.LeftInline(LastPathSeparator, EAllowShrinking::No);
+
+				FSoftObjectPath ActorPath = FSoftObjectPath::ConstructFromAssetPathAndSubpath(AssetPath, MoveTemp(SubPath));
+				TSoftObjectPtr<UObject> ActorPtr(MoveTemp(ActorPath));
+
+				return ActorPtr;
+			}
+		}
+		return {};
+	}
+#endif //WITH_EDITOR
 }
 
 // TODO: Not sure if using this kind of constructor is a proper thing to do vs some other hook...
@@ -267,7 +299,7 @@ TArray<UE::Landscape::EditLayers::FEditLayerRendererState> ALandscapePatchManage
 void ALandscapePatchManager::SetTargetLandscape(ALandscape* InTargetLandscape)
 {
 #if WITH_EDITOR
-	if (OwningLandscape != InTargetLandscape)
+	if (OwningLandscape != InTargetLandscape && !bDead)
 	{
 		if (OwningLandscape)
 		{
@@ -326,8 +358,9 @@ bool ALandscapePatchManager::ContainsPatch(ULandscapePatchComponent* Patch) cons
 void ALandscapePatchManager::AddPatch(ULandscapePatchComponent* Patch)
 {
 	using namespace LandscapePatchManagerLocals;
+	ensure(!IsDead());
 
-	if (Patch && Patch->IsPatchInWorld())
+	if (Patch && Patch->IsPatchInWorld() && !IsDead())
 	{
 		if (!ContainsPatch(Patch))
 		{
@@ -439,11 +472,22 @@ void ALandscapePatchManager::MovePatchToIndex(ULandscapePatchComponent* Patch, i
 }
 
 #if WITH_EDITOR
+
 void ALandscapePatchManager::MigrateToPrioritySystemAndDelete()
 {
-	if (!PatchComponents.IsEmpty())
+	MigrateToPrioritySystemAndDeleteInternal(/*bAllowUI=*/true);
+}
+
+void ALandscapePatchManager::MigrateToPrioritySystemAndDeleteInternal(bool bAllowUI)
+{
+	check(GetWorld()->WorldType == EWorldType::Editor);
+	
+	// Create LandscapeInfo if needed.  During load time, it might not exist yet, depending on load order.
+	ULandscapeInfo* LandscapeInfo = OwningLandscape ? OwningLandscape->CreateLandscapeInfo(true) : nullptr;
+	if (OwningLandscape && LandscapeInfo && !bDead && IsValid(this) && !PatchComponents.IsEmpty())
 	{
 		Modify();
+		LandscapeInfo->MarkObjectDirty(OwningLandscape, /*bInForceResave = */true);
 
 		// Patches will remove themselves from PatchComponents as we go along, so we need to iterate
 		// a copy.
@@ -451,6 +495,7 @@ void ALandscapePatchManager::MigrateToPrioritySystemAndDelete()
 
 		// We call Modify on all the patches we'll be touching at the start, otherwise they will
 		// store incorrect indices for undo as they are removed.
+		int32 PatchErrors = 0;
 		for (TSoftObjectPtr<ULandscapePatchComponent> Patch : PatchComponents)
 		{
 			if (Patch.IsPending())
@@ -458,11 +503,51 @@ void ALandscapePatchManager::MigrateToPrioritySystemAndDelete()
 				Patch.LoadSynchronous();
 			}
 
+			if (!Patch.IsValid())
+			{
+				// Loading directly from a component TSoftObjectPtr is unreliable, so try deriving the owner actor pointer using
+				// the path and load via that.  If the setup is weird and this doesn't work, it should at least still be safe. 
+				TSoftObjectPtr<UObject> ActorPtr = LandscapePatchManagerLocals::GetActorPtrFromPatchComponent(Patch);
+				if (ActorPtr.IsPending())
+				{
+					ActorPtr.LoadSynchronous();  // If this succeeds, Patch will turn valid.
+				}
+			}
+
 			if (Patch.IsValid())
 			{
+				LandscapeInfo->MarkObjectDirty(Patch.Get(), /*bInForceResave = */true, OwningLandscape);
 				Patch->Modify();
 				PatchListCopy.Add(Patch);
 			}
+			else
+			{
+				// Failed to load?  This can happen if patch bindings were broken in the pre-migration scene.
+				// LandscapePatch.FixPatchBindings after deleting the patch manager can fix them.
+				++PatchErrors;
+			}
+		}
+
+		if (!bAllowUI)
+		{
+			// Create the ULandscapePatchEditLayer in advance, if needed.  This prevents BindToLandscape from triggering the
+			// modal UI window for choosing the new layer index.
+			const FLandscapeLayer* LayerStruct = OwningLandscape->FindLayerOfType(ULandscapePatchEditLayer::StaticClass());
+			ULandscapePatchEditLayer* Layer = LayerStruct ? Cast<ULandscapePatchEditLayer>(LayerStruct->EditLayer) : nullptr;
+
+			if (!Layer)
+			{
+				const FName PatchLayerName = OwningLandscape->GenerateUniqueLayerName(*ULandscapePatchEditLayer::StaticClass()->GetDefaultObject<ULandscapePatchEditLayer>()->GetDefaultName());
+				int32 LayerIdx = OwningLandscape->CreateLayer(PatchLayerName, ULandscapePatchEditLayer::StaticClass());
+				check(LayerIdx != INDEX_NONE);
+
+				// Move the new layer beside the old-style "LandscapePatches" layer
+				int32 OldPathLayerIdx = OwningLandscape->GetBrushLayer(this);
+				if (ensure(OldPathLayerIdx != INDEX_NONE))
+				{
+					OwningLandscape->ReorderLayer(LayerIdx, OldPathLayerIdx);
+				}
+			}
 		}
 
 		double Priority = LEGACY_PATCH_PRIORITY_BASE;
@@ -477,7 +562,14 @@ void ALandscapePatchManager::MigrateToPrioritySystemAndDelete()
 				Patch->SetPatchManager(nullptr);
 			PRAGMA_ENABLE_DEPRECATION_WARNINGS
 
-				Patch->FixBindings();
+			Patch->FixBindings();
+		}
+
+		UE_LOG(LogLandscapePatch, Warning, TEXT("ALandscapePatchManager: %d landscape patches have been migrated from the legacy patch manager \"%s\" to be bound directly to a ULandscapePatchEditLayer"),
+			PatchListCopy.Num(), *GetActorLabel(false));
+		if (PatchErrors > 0)
+		{
+			UE_LOG(LogLandscapePatch, Error, TEXT("ALandscapePatchManager: %d landscape patches failed to migrate successfully.  They may be restorable by running the command LandscapePatch.FixPatchBindings"), PatchErrors);
 		}
 
 		PatchComponents.Empty();
@@ -487,9 +579,24 @@ void ALandscapePatchManager::MigrateToPrioritySystemAndDelete()
 	// Important so that we remove ourselves from the landscape blueprint brush list
 	SetTargetLandscape(nullptr);
 
-	// This will warn the user if there are still references to the manager
-	UTypedElementSelectionSet* SelectionSet = GEditor->GetEditorSubsystem<ULevelEditorSubsystem>()->GetSelectionSet();
-	GUnrealEd->DeleteActors({ this }, GetWorld(), SelectionSet);
+	// bDead is used as protection from unexpected weirdness if anything happens in the window between this code running and the
+	// actor being actually removed from the world (by the actionable message update button).
+	bDead = true;
+
+	if (LandscapeInfo)
+	{
+		// We can't delete an actor during load time, so enqueue this to be deleted later.  We want this deletion to be applied from
+		// MarkModifiedLandscapesAsDirty, the same place that the deferred dirty state from MarkObjectDirty is finally applied.  This
+		// will leave the scene in a consistent state before and after the user clicks the "Update" button from the landscape check
+		// code.  If the scene is closed without clicking update, it will remain fully un-migrated on disk.
+		LandscapeInfo->DeleteActorWhenApplyingModifiedStatus(this);
+	}
+	else
+	{
+		// No LandscapeInfo, likely because no OwningLandscape.  Try to delete directly.
+		UTypedElementSelectionSet* SelectionSet = GEditor->GetEditorSubsystem<ULevelEditorSubsystem>()->GetSelectionSet();
+		GUnrealEd->DeleteActors({ this }, GetWorld(), SelectionSet);
+	}
 }
 
 
@@ -617,6 +724,17 @@ void ALandscapePatchManager::PostLoad()
 	}
 }
 
+void ALandscapePatchManager::PostRegisterAllComponents()
+{
+	Super::PostRegisterAllComponents();
+
+	if (!IsTemplate() && IsValid(GetWorld()) && GetWorld()->WorldType == EWorldType::Editor &&
+		LandscapePatchManagerLocals::CVarMigrateLegacyPatchListToPrioritySystem->GetBool())
+	{
+		MigrateToPrioritySystemAndDeleteInternal(/*bAllowUI=*/false);
+	}
+}
+
 void ALandscapePatchManager::CheckForErrors()
 {
 	using namespace LandscapePatchManagerLocals;
diff --git a/Engine/Plugins/Experimental/LandscapePatch/Source/LandscapePatch/Public/LandscapePatchManager.h b/Engine/Plugins/Experimental/LandscapePatch/Source/LandscapePatch/Public/LandscapePatchManager.h
--- a/Engine/Plugins/Experimental/LandscapePatch/Source/LandscapePatch/Public/LandscapePatchManager.h
+++ b/Engine/Plugins/Experimental/LandscapePatch/Source/LandscapePatch/Public/LandscapePatchManager.h
@@ -87,6 +87,15 @@ public:
 	UFUNCTION(BlueprintCallable, Category = LandscapePatch)
 	void MovePatchToIndex(ULandscapePatchComponent* Patch, int32 Index);
 
+	// This patch manager has been migrated out of and should no longer be accessible.
+	bool IsDead() const {
+#if WITH_EDITOR
+		return bDead;
+#else
+		return false;
+#endif
+	}
+
 #if WITH_EDITOR
 
 	/**
@@ -108,6 +117,7 @@ public:
 
 	// AActor
 	virtual void CheckForErrors() override;
+	virtual void PostRegisterAllComponents() override;
 
 	// UObject
 	virtual void PostEditUndo() override;
@@ -123,6 +133,8 @@ public:
 	inline static const double LEGACY_PATCH_PRIORITY_BASE = ULandscapePatchEditLayer::PATCH_PRIORITY_BASE - 10;
 private:
 
+	void MigrateToPrioritySystemAndDeleteInternal(bool bAllowUI);
+
 	UPROPERTY()
 	TArray<TSoftObjectPtr<ULandscapePatchComponent>> PatchComponents;
 
@@ -139,6 +151,8 @@ private:
 	TObjectPtr<ALandscape> DetailPanelLandscape = nullptr;
 
 	bool bIssuedPatchOwnershipWarning = false;
+
+	bool bDead = false;
 #endif
 
 	// Transient table to speed up Contains and IndexOf queries, which are very slow for an array of TSoftObjectPtr's.
diff --git a/Engine/Source/Runtime/Landscape/Classes/LandscapeInfo.h b/Engine/Source/Runtime/Landscape/Classes/LandscapeInfo.h
--- a/Engine/Source/Runtime/Landscape/Classes/LandscapeInfo.h
+++ b/Engine/Source/Runtime/Landscape/Classes/LandscapeInfo.h
@@ -167,14 +167,13 @@ public:
 
 private:
 #if WITH_EDITORONLY_DATA
-	TSet<TWeakObjectPtr<UObject>, TWeakObjectPtrSetKeyFuncs<TWeakObjectPtr<UObject>>> PackagesToDeprecate;
-
 	// SORTED list of all actors implementing the spline interface that are registered with this landscape info
 	UPROPERTY()
 	TArray<TScriptInterface<ILandscapeSplineInterface>> SplineActors;
 
-	// Not a property since this shouldn't be modified through transactions (no undo/redo)
+	// Not properties since these shouldn't be modified through transactions (no undo/redo)
 	TSet<TWeakObjectPtr<UPackage>, TWeakObjectPtrSetKeyFuncs<TWeakObjectPtr<UPackage>>> ModifiedPackages;
+	TSet<TWeakObjectPtr<AActor>, TWeakObjectPtrSetKeyFuncs<TWeakObjectPtr<AActor>>> ActorsToDelete;
 
 	bool bDirtyOnlyInMode;
 
@@ -289,7 +288,7 @@ public:
 	LANDSCAPE_API bool ModifyObject(UObject* InObject, bool bAlwaysMarkDirty = true);
 	/**
 	 * See UObjectBaseUtility::MarkPackageDirty
-	 * @param InObject Object to mark dirty. Must be either a ALandscapeProxy or one of its inner objects
+	 * @param InObject Object to mark dirty.
 	 * @param bInForceResave passed to true in order to bypass the bDirtyOnlyInMode and force the object into the list of modified packages. This ensures the object is marked as needing resaving in case we want to 
 	 *  help the user re-save its landscape actors (e.g. following a data fixup)
 	 * @param InLandscapeOverride if set, MarkObjectDirty will use this ALandscape instead of the local one in the method.
@@ -297,6 +296,12 @@ public:
 	 */
 	LANDSCAPE_API bool MarkObjectDirty(UObject* InObject, bool bInForceResave = false, const ALandscape* InLandscapeOverride = nullptr);
 
+	/**
+	 * Enqueue an actor to be deleted when deferred dirty status from MarkObjectDirty is finally applied.
+	 * @param InActor Actor to delete.
+	 */
+	LANDSCAPE_API void DeleteActorWhenApplyingModifiedStatus(AActor* InActor);
+
 	/** Helper function to invalidate the runtime virtual texture data associated with the specified landscape region. 
 	 *   This forces the dirtied RVT pages to be updated in priority
 	 */
@@ -464,6 +469,7 @@ private:
 	void MoveControlPoint(ULandscapeSplineControlPoint* InControlPoint, TScriptInterface<ILandscapeSplineInterface> From, TScriptInterface<ILandscapeSplineInterface> To);
 	bool UpdateLayerInfoMapInternal(ALandscapeProxy* Proxy);
 	bool TryAddToModifiedPackages(UPackage* InPackage, const ALandscape* InLandscapeOverride = nullptr);
+	void ProcessDeferredDeletions();
 #endif
 
 private:
diff --git a/Engine/Source/Runtime/Landscape/Private/Landscape.cpp b/Engine/Source/Runtime/Landscape/Private/Landscape.cpp
--- a/Engine/Source/Runtime/Landscape/Private/Landscape.cpp
+++ b/Engine/Source/Runtime/Landscape/Private/Landscape.cpp
@@ -115,6 +115,7 @@ Landscape.cpp: Terrain rendering
 #include "MaterialUtilities.h"
 #include "Editor.h"
 #include "Editor/EditorEngine.h"
+#include "Engine/Selection.h"
 #include "Engine/Texture2D.h"
 #include "AssetCompilingManager.h"
 #include "FileHelpers.h"
@@ -5151,6 +5152,8 @@ int32 ULandscapeInfo::MarkModifiedPackagesAsDirty()
 	}
 	ModifiedPackages.Empty();
 
+	ProcessDeferredDeletions();
+
 	return NumDirtied;
 }
 
@@ -5194,7 +5197,7 @@ bool ULandscapeInfo::TryAddToModifiedPackages(UPackage* InPackage, const ALandsc
 
 bool ULandscapeInfo::MarkObjectDirty(UObject* InObject, bool bInForceResave, const ALandscape* InLandscapeOverride)
 {
-	check(InObject && (InObject->IsA<ALandscapeProxy>() || InObject->GetTypedOuter<ALandscapeProxy>() != nullptr));
+	check(InObject);
 		
 	bool bWasAddedToModifiedPackages = false;
 	if (bInForceResave)
@@ -5260,6 +5263,33 @@ bool ULandscapeInfo::ModifyObject(UObject* InObject, bool bAlwaysMarkDirty)
 	return bWasAddedToModifiedPackages;
 }
 
+void ULandscapeInfo::DeleteActorWhenApplyingModifiedStatus(AActor* InActor)
+{
+	check(InActor->GetWorld() && InActor->GetWorld()->WorldType == EWorldType::Editor);
+
+	// If we can mark the package dirty, then we can also delete it right away.  If we can't, then enqueue it for deletion at the same time as other deferred package dirtying.
+	if (InActor->MarkPackageDirty())
+	{
+		GUnrealEd->DeleteActors({ InActor }, InActor->GetWorld(), GUnrealEd->GetSelectedActors()->GetElementSelectionSet());
+	}
+	else
+	{
+		ActorsToDelete.Add(InActor);
+	}
+}
+
+void ULandscapeInfo::ProcessDeferredDeletions()
+{
+	for (TWeakObjectPtr<AActor> WeakActorPtr : ActorsToDelete)
+	{
+		if (AActor* Actor = WeakActorPtr.Get())
+		{
+			GUnrealEd->DeleteActors({ Actor }, Actor->GetWorld(), GUnrealEd->GetSelectedActors()->GetElementSelectionSet());
+		}
+	}
+	ActorsToDelete.Empty();
+}
+
 void ULandscapeInfo::DirtyRuntimeVirtualTextureForLandscapeArea(int32 X1, int32 Y1, int32 X2, int32 Y2) const
 {
 	FBox DirtyWorldBounds(ForceInit);
diff --git a/Engine/Source/Runtime/Landscape/Private/LandscapeEditLayers.cpp b/Engine/Source/Runtime/Landscape/Private/LandscapeEditLayers.cpp
--- a/Engine/Source/Runtime/Landscape/Private/LandscapeEditLayers.cpp
+++ b/Engine/Source/Runtime/Landscape/Private/LandscapeEditLayers.cpp
@@ -12259,9 +12259,7 @@ void ALandscapeProxy::DeleteLayer(const FGuid& InLayerGuid)
 
 void ALandscapeProxy::InitializeLayerWithEmptyContent(const FGuid& InLayerGuid)
 {
-	ULandscapeInfo* Info = GetLandscapeInfo();
-
-	if (Info == nullptr)
+	if (IsPendingKillPending() || !GetLandscapeActor() || !LandscapeGuid.IsValid())
 	{
 		return;
 	}
@@ -13220,12 +13218,19 @@ void ALandscape::OnLayerCreatedInternal(FLandscapeLayer& Layer)
 	Layer.EditLayer->SetFlags(RF_Transactional);
 	Layer.EditLayer->OnLayerCreated(Layer);
 
+	// Add to self first instead of from the loop.  Avoids potential load-order problems if this is called at load-time,
+	// before the LandscapeInfo is created.  The other proxies can fix themselves if they load later.
+	AddLayer(Layer.Guid);
+
 	// Create associated layer data in each landscape proxy
 	if (ULandscapeInfo* LandscapeInfo = GetLandscapeInfo())
 	{
-		LandscapeInfo->ForEachLandscapeProxy([&Layer](ALandscapeProxy* Proxy)
+		LandscapeInfo->ForEachLandscapeProxy([&Layer, this](ALandscapeProxy* Proxy)
 		{
-			Proxy->AddLayer(Layer.Guid);
+			if (Proxy != this)
+			{
+				Proxy->AddLayer(Layer.Guid);
+			}
 			return true;
 		});
 	}
