Fix non-unity linking due to missing template instantiation
#rnx

diff --git a/Engine/Plugins/Experimental/LandscapePatch/Source/LandscapePatch/Private/LandscapeCircleHeightPatchPS.cpp b/Engine/Plugins/Experimental/LandscapePatch/Source/LandscapePatch/Private/LandscapeCircleHeightPatchPS.cpp
--- a/Engine/Plugins/Experimental/LandscapePatch/Source/LandscapePatch/Private/LandscapeCircleHeightPatchPS.cpp
+++ b/Engine/Plugins/Experimental/LandscapePatch/Source/LandscapePatch/Private/LandscapeCircleHeightPatchPS.cpp
@@ -67,3 +67,9 @@ IMPLEMENT_SHADER_TYPE(template<>, FLandscapeCircleHeightPatchPS</*bPerformBlendi
 IMPLEMENT_SHADER_TYPE(template<>, FLandscapeCircleHeightPatchPS</*bPerformBlending = */true>, TEXT("/Plugin/LandscapePatch/Private/LandscapeCircleHeightPatchPS.usf"), TEXT("ApplyLandscapeCircleHeightPatch"), SF_Pixel);
 IMPLEMENT_SHADER_TYPE(template<>, FLandscapeCircleVisibilityPatchPS</*bPerformBlending = */false>, TEXT("/Plugin/LandscapePatch/Private/LandscapeCircleHeightPatchPS.usf"), TEXT("ApplyLandscapeCircleVisibilityPatch"), SF_Pixel);
 IMPLEMENT_SHADER_TYPE(template<>, FLandscapeCircleVisibilityPatchPS</*bPerformBlending = */true>, TEXT("/Plugin/LandscapePatch/Private/LandscapeCircleHeightPatchPS.usf"), TEXT("ApplyLandscapeCircleVisibilityPatch"), SF_Pixel);
+
+// Explicit instantiation otherwise non-unity builds cannot link
+template class FLandscapeCircleHeightPatchPS</*bPerformBlending = */false>;
+template class FLandscapeCircleHeightPatchPS</*bPerformBlending = */true>;
+template class FLandscapeCircleVisibilityPatchPS</*bPerformBlending = */false>;
+template class FLandscapeCircleVisibilityPatchPS</*bPerformBlending = */true>;
\ No newline at end of file
