Avoid crash when opening a map with a disconnected landscape patch manager, when the editor isn't ready to provide the current selection set.

diff --git a/Engine/Plugins/Experimental/LandscapePatch/Source/LandscapePatch/Private/LandscapePatchManager.cpp b/Engine/Plugins/Experimental/LandscapePatch/Source/LandscapePatch/Private/LandscapePatchManager.cpp
--- a/Engine/Plugins/Experimental/LandscapePatch/Source/LandscapePatch/Private/LandscapePatchManager.cpp
+++ b/Engine/Plugins/Experimental/LandscapePatch/Source/LandscapePatch/Private/LandscapePatchManager.cpp
@@ -566,12 +566,16 @@ void ALandscapePatchManager::MigrateToPrioritySystemAndDeleteInternal(bool bAllo
 		// code.  If the scene is closed without clicking update, it will remain fully un-migrated on disk.
 		LandscapeInfo->DeleteActorWhenApplyingModifiedStatus(this);
 	}
-	else
+	else if (UTypedElementSelectionSet* SelectionSet = GEditor->GetEditorSubsystem<ULevelEditorSubsystem>()->GetSelectionSet())
 	{
 		// No LandscapeInfo, likely because no OwningLandscape.  Try to delete directly.
-		UTypedElementSelectionSet* SelectionSet = GEditor->GetEditorSubsystem<ULevelEditorSubsystem>()->GetSelectionSet();
 		GUnrealEd->DeleteActors({ this }, GetWorld(), SelectionSet);
 	}
+	else
+	{
+		// No SelectionSet because Editor isn't ready yet.  Print a message and leave the disconnected patch manager in place.
+		UE_LOG(LogLandscapePatch, Display, TEXT("ALandscapePatchManager \"%s\" is unused and can be deleted."), *GetActorLabel(false));
+	}
 }
 
 
