Avoid crash on opening a scene containing the legacy LandscapePatchManager when the scene was already at the max number of edit layers.  Ignore the soft limit when creating the new dedicated patch layer.

diff --git a/Engine/Plugins/Experimental/LandscapePatch/Source/LandscapePatch/Private/LandscapePatchManager.cpp b/Engine/Plugins/Experimental/LandscapePatch/Source/LandscapePatch/Private/LandscapePatchManager.cpp
--- a/Engine/Plugins/Experimental/LandscapePatch/Source/LandscapePatch/Private/LandscapePatchManager.cpp
+++ b/Engine/Plugins/Experimental/LandscapePatch/Source/LandscapePatch/Private/LandscapePatchManager.cpp
@@ -334,7 +334,6 @@ bool ALandscapePatchManager::ContainsPatch(ULandscapePatchComponent* Patch) cons
 void ALandscapePatchManager::AddPatch(ULandscapePatchComponent* Patch)
 {
 	using namespace LandscapePatchManagerLocals;
-	ensure(!IsDead());
 
 	if (Patch && Patch->IsPatchInWorld() && !IsDead())
 	{
@@ -362,6 +361,12 @@ void ALandscapePatchManager::AddPatch(ULandscapePatchComponent* Patch)
 				RequestLandscapeUpdate(!UE::GetIsEditorLoadingPackage());
 			}
 	}
+	else if (Patch && IsDead())
+	{
+		// IsDead means the migration code ran.  A patch showing up here means it wasn't successfully migrated, possibly from failing to load.
+		UE_LOG(LogLandscapePatch, Error, TEXT("Landscape patch is not correctly connected to a LandscapePatchEditLayer.  Fix this by first clicking Build->Build Landscape and then running LandscapePatch.FixPatchBindings. "
+			"(Package: %s, Actor : %s)"), *GetNameSafe(Patch->GetPackage()), *GetNameSafe(Patch->GetAttachmentRootActor()));
+	}
 }
 
 bool ALandscapePatchManager::RemovePatch(ULandscapePatchComponent* Patch)
@@ -513,7 +518,9 @@ void ALandscapePatchManager::MigrateToPrioritySystemAndDeleteInternal(bool bAllo
 			if (!Layer)
 			{
 				const FName PatchLayerName = OwningLandscape->GenerateUniqueLayerName(*ULandscapePatchEditLayer::StaticClass()->GetDefaultObject<ULandscapePatchEditLayer>()->GetDefaultName());
-				int32 LayerIdx = OwningLandscape->CreateLayer(PatchLayerName, ULandscapePatchEditLayer::StaticClass());
+				// Ignore the layer count limit to avoid failing here.  It's only a soft limit to aid performance.  The old non-ULandscapePatchEditLayer is probably unused after this, but
+				// it could have also been used for manual painting and we can't safely delete it.
+				int32 LayerIdx = OwningLandscape->CreateLayer(PatchLayerName, ULandscapePatchEditLayer::StaticClass(), /*bIgnoreLayerCountLimit=*/ true);
 				check(LayerIdx != INDEX_NONE);
 
 				// Move the new layer beside the old-style "LandscapePatches" layer
diff --git a/Engine/Source/Runtime/Landscape/Classes/Landscape.h b/Engine/Source/Runtime/Landscape/Classes/Landscape.h
--- a/Engine/Source/Runtime/Landscape/Classes/Landscape.h
+++ b/Engine/Source/Runtime/Landscape/Classes/Landscape.h
@@ -380,9 +380,10 @@ public:
 	* Creates a new edit layer
 	* @Param InName is the name of the new edit layer
 	* @Param InEditLayerClass is the class of the edit layer to create. Passing null will create a standard layer of type ULandscapeEditLayer
+	* @Param bInIgnoreLayerCountLimit create the layer even if this would exceed MaxNumberOfLayers
 	* @return the index of the newly-created layer
 	*/
-	LANDSCAPE_API int32 CreateLayer(FName InName = NAME_None, const TSubclassOf<ULandscapeEditLayerBase>& InEditLayerClass = TSubclassOf<ULandscapeEditLayerBase>());
+	LANDSCAPE_API int32 CreateLayer(FName InName = NAME_None, const TSubclassOf<ULandscapeEditLayerBase>& InEditLayerClass = TSubclassOf<ULandscapeEditLayerBase>(), bool bInIgnoreLayerCountLimit = false);
 
 	/** 
 	* Creates a new edit layer by copying all the settings from the layer passed in parameter
diff --git a/Engine/Source/Runtime/Landscape/Private/LandscapeEditLayers.cpp b/Engine/Source/Runtime/Landscape/Private/LandscapeEditLayers.cpp
--- a/Engine/Source/Runtime/Landscape/Private/LandscapeEditLayers.cpp
+++ b/Engine/Source/Runtime/Landscape/Private/LandscapeEditLayers.cpp
@@ -13282,7 +13282,7 @@ FLandscapeLayer* ALandscape::DuplicateLayerAndMoveBrushes(const FLandscapeLayer&
 	return &LandscapeEditLayers[AddedIndex];
 }
 
-int32 ALandscape::CreateLayer(FName InName, const TSubclassOf<ULandscapeEditLayerBase>& InEditLayerClass)
+int32 ALandscape::CreateLayer(FName InName, const TSubclassOf<ULandscapeEditLayerBase>& InEditLayerClass, bool bInIgnoreLayerCountLimit)
 {
 	// Detect any attempt to call this in the middle of UpdateLayersContent.  If called from blueprint, log an error and return early instead of asserting.
 	if (InLayerUpdateCount > 0 && UE::Landscape::Private::InBPCallstack())
@@ -13292,7 +13292,7 @@ int32 ALandscape::CreateLayer(FName InName, const TSubclassOf<ULandscapeEditLaye
 	}
 	check(InLayerUpdateCount == 0);
 
-	if (IsMaxLayersReached() || !CanHaveLayersContent())
+	if ((!bInIgnoreLayerCountLimit && IsMaxLayersReached()) || !CanHaveLayersContent())
 	{
 		return INDEX_NONE;
 	}
