* Fixed several issues with ReinitializeHeight/Weights in landscape patches
** Fixed crash with multiple patches. The render command recorder was not stopped at the right moment : both Render and Blend operations must respect the render mode (Immediate or Recorded), which could be broken in case the renderer was grouped with another one earlier.
** Fixed landscape that was not updated after having reinitialization had been executed
** Fixed crash when the same patch is rendered in several batches. The boolean bReinitializeHeightOnNextRender does not support this and that supposes that during reinitialization, the patch is rendered in a single batch. The fix artificially increases the batch size to the entire landscape, which is wasteful, but the best we can do until the reinitialization is replaced by a proper landscape edit layer partial merge.
* Fixed landscape batched merge not handling multiple weight patches (i.e. multiple target layers) being rendered on the same patch component (the render targets would cycle in between each, leading to invalid blending)
* Fixed landscape patches' input area (which was the size of the patch), that would lead to needlessly large batches
* Fixed potential determinism issue with landscape patches that could theoretically render before their referenced textures are fully streamed in

Misc:
* Added a simplified version of CVar landscape.BatchedMerge.VisualLog.ShowMergeProcess (1 = show batches only, 2 = show all, as before)

diff --git a/Engine/Plugins/Experimental/LandscapePatch/Source/LandscapePatch/Private/LandscapeTexturePatch.cpp b/Engine/Plugins/Experimental/LandscapePatch/Source/LandscapePatch/Private/LandscapeTexturePatch.cpp
--- a/Engine/Plugins/Experimental/LandscapePatch/Source/LandscapePatch/Private/LandscapeTexturePatch.cpp
+++ b/Engine/Plugins/Experimental/LandscapePatch/Source/LandscapePatch/Private/LandscapeTexturePatch.cpp
@@ -22,6 +22,8 @@
 #include "SystemTextures.h"
 #include "TextureResource.h"
 #include "Algo/AnyOf.h"
+#include "Algo/ForEach.h"
+#include "Containers/Ticker.h"
 
 #include UE_INLINE_GENERATED_CPP_BY_NAME(LandscapeTexturePatch)
 
@@ -133,8 +135,11 @@ UE::Landscape::EditLayers::ERenderFlags ULandscapeTexturePatch::GetRenderFlags(c
 	ERenderFlags RenderFlags = ERenderFlags::None;
 	// COMMENT [jonathan.bard] : this is not something we want to keep (we will use partial edit layer renders for this eventually) but we can still render in immediate mode
 	//  in the "reinitialize on next render case" because we perform a synchronous read then so we need to run on the game thread to perform the rendering commands flush 
-	if ((InMergeContext->IsHeightmapMerge() && bReinitializeHeightOnNextRender)
-		|| Algo::AnyOf(WeightPatches, [](const TObjectPtr<ULandscapeWeightPatchTextureInfo>& Info) { return Info->bReinitializeOnNextRender; }))
+	if (InMergeContext->IsHeightmapMerge() && bReinitializeHeightOnNextRender)
+	{
+		RenderFlags |= ERenderFlags::RenderMode_Immediate;
+	}
+	else if (!InMergeContext->IsHeightmapMerge() && Algo::AnyOf(WeightPatches, [](const TObjectPtr<ULandscapeWeightPatchTextureInfo>& Info) { return Info->bReinitializeOnNextRender; }))
 	{
 		RenderFlags |= ERenderFlags::RenderMode_Immediate;
 	}
@@ -183,9 +188,9 @@ bool ULandscapeTexturePatch::RenderLayer(UE::Landscape::EditLayers::FRenderParam
 		else
 		{
 			checkf(RDGBuilderRecorder.IsRecording(), TEXT("We should be using ERenderFlags::RenderMode_Recorded in the typical case"));
-			bool bDidRenderSomething = false;
-			ApplyToHeightmap(/*bInPerformBlending = */false, &RenderParams, RDGBuilderRecorder, CurrentData, GetHeightmapToWorld(RenderParams.RenderAreaWorldTransform), bDidRenderSomething);
-			return bDidRenderSomething;
+			bool bHasRenderedSomething = false;
+			ApplyToHeightmap(/*bInPerformBlending = */false, &RenderParams, RDGBuilderRecorder, CurrentData, GetHeightmapToWorld(RenderParams.RenderAreaWorldTransform), bHasRenderedSomething);
+			return bHasRenderedSomething;
 		}
 	}
 
@@ -197,7 +202,7 @@ bool ULandscapeTexturePatch::RenderLayer(UE::Landscape::EditLayers::FRenderParam
 	int32 NumTargetLayersInGroup = RenderParams.TargetLayerGroupLayerNames.Num();
 	check(LandscapeScratchRT->GetEffectiveNumSlices() == NumTargetLayersInGroup);
 
-	bool bDidRenderSomething = false;
+	bool bHasRenderedSomething = false;
 
 	for (int32 TargetLayerIndexInGroup = 0; TargetLayerIndexInGroup < NumTargetLayersInGroup; ++TargetLayerIndexInGroup)
 	{
@@ -229,24 +234,21 @@ bool ULandscapeTexturePatch::RenderLayer(UE::Landscape::EditLayers::FRenderParam
 			ReinitializeWeightPatch(WeightPatchInfo, TextureArray->GetResource(),
 				FIntPoint(TextureArray->SizeX, TextureArray->SizeY), TargetLayerIndexInGroup, LandscapeHeightmapToWorld);
 
-			bDidRenderSomething = true;
+			bHasRenderedSomething = true;
 		}
 		else
 		{
 			checkf(RDGBuilderRecorder.IsRecording(), TEXT("We should be using ERenderFlags::RenderMode_Recorded in the typical case"));
-			bool bLayerDidRenderSomething = false;
 			ApplyToWeightmap(/*bInPerformBlending = */false, &RenderParams, RDGBuilderRecorder, WeightPatchInfo,
 				TextureArray->GetResource(),
 				TargetLayerIndexInGroup,
 				RenderParams.RenderAreaSectionRect.Size(), 
 				GetHeightmapToWorld(RenderParams.RenderAreaWorldTransform), 
-				bLayerDidRenderSomething);
-
-			bDidRenderSomething |= bLayerDidRenderSomething;
+				bHasRenderedSomething);
 		}
 	}//end for each layer index
 
-	return bDidRenderSomething;
+	return bHasRenderedSomething;
 }//end ULandscapeTexturePatch::RenderLayer
 
 void ULandscapeTexturePatch::BlendLayer(UE::Landscape::EditLayers::FRenderParams& RenderParams, UE::Landscape::FRDGBuilderRecorder& RDGBuilderRecorder)
@@ -300,8 +302,8 @@ UTextureRenderTarget2D* ULandscapeTexturePatch::RenderLayer_Native(const FLandsc
 		{
 			// We render in immediate mode for the legacy path so we use an immediate recorder :  
 			FRDGBuilderRecorder RDGBuilderRecorderImmediate;
-			bool bDidRenderSomething = false;
-			return ApplyToHeightmap(/*bInPerformBlending = */true, /*RenderParams = */nullptr, RDGBuilderRecorderImmediate, InParameters.CombinedResult, LandscapeHeightmapToWorld, bDidRenderSomething, /*OutputAccess = */ERHIAccess::SRVMask); // in the legacy path the RT leaves as a SRV
+			bool bHasRenderedSomething = false;
+			return ApplyToHeightmap(/*bInPerformBlending = */true, /*RenderParams = */nullptr, RDGBuilderRecorderImmediate, InParameters.CombinedResult, LandscapeHeightmapToWorld, bHasRenderedSomething, /*OutputAccess = */ERHIAccess::SRVMask); // in the legacy path the RT leaves as a SRV
 		}
 	}
 	else
@@ -345,13 +347,13 @@ UTextureRenderTarget2D* ULandscapeTexturePatch::RenderLayer_Native(const FLandsc
 		{
 			// We render in immediate mode for the legacy path so we use an immediate recorder :  
 			FRDGBuilderRecorder RDGBuilderRecorderImmediate;
-			bool bLayerDidRenderSomething = false;
+			bool bHasRenderedSomething = false;
 			ApplyToWeightmap(/*bInPerformBlending = */true, /*RenderParams = */nullptr, RDGBuilderRecorderImmediate, WeightPatchInfo,
 				InParameters.CombinedResult->GetResource(),
 				0, // Slice index
 				FIntPoint(InParameters.CombinedResult->SizeX, InParameters.CombinedResult->SizeY),
 				LandscapeHeightmapToWorld, 
-				bLayerDidRenderSomething,
+				bHasRenderedSomething,
 				/*OutputAccess = */ERHIAccess::SRVMask); // in the legacy path the RT leaves as a SRV
 			return InParameters.CombinedResult;
 		}
@@ -359,10 +361,9 @@ UTextureRenderTarget2D* ULandscapeTexturePatch::RenderLayer_Native(const FLandsc
 }
 
 UTextureRenderTarget2D* ULandscapeTexturePatch::ApplyToHeightmap(bool bInPerformBlending, UE::Landscape::EditLayers::FRenderParams* RenderParams, UE::Landscape::FRDGBuilderRecorder& RDGBuilderRecorder, UTextureRenderTarget2D* InCombinedResult,
-	const FTransform& LandscapeHeightmapToWorld, bool& bOutDidRenderSomething, ERHIAccess OutputAccess)
+	const FTransform& LandscapeHeightmapToWorld, bool& bHasRenderedSomething, ERHIAccess OutputAccess)
 {
 	using namespace UE::Landscape;
-	bOutDidRenderSomething = false;
 
 	// Get the source of our height patch
 	UTexture* PatchUObject = nullptr;
@@ -420,7 +421,7 @@ UTextureRenderTarget2D* ULandscapeTexturePatch::ApplyToHeightmap(bool bInPerform
 		ULandscapeScratchRenderTarget* WriteRT = RenderParams->MergeRenderContext->GetBlendRenderTargetWrite();
 		// After this point, the render cannot fail so if we're the first in our render layer group to render, we can cycle the blend render targets and 
 		//  start rendering in the write one : 
-		if (RenderParams->NumSuccessfulRenderLayerStepsUntilBlendLayerStep == 0)
+		if (!bHasRenderedSomething && (RenderParams->NumSuccessfulRenderLayerStepsUntilBlendLayerStep == 0))
 		{
 			RenderParams->MergeRenderContext->CycleBlendRenderTargets(RDGBuilderRecorder);
 			WriteRT = RenderParams->MergeRenderContext->GetBlendRenderTargetWrite();
@@ -496,16 +497,15 @@ UTextureRenderTarget2D* ULandscapeTexturePatch::ApplyToHeightmap(bool bInPerform
 	// We need to specify the final state of the external texture to prevent the graph builder from transitioning it to SRVMask :
 	RDGBuilderRecorder.EnqueueRDGCommand(RDGCommand, { { OutputResource, (OutputAccess == ERHIAccess::None) ? ERHIAccess::RTV : OutputAccess} });
 
-	bOutDidRenderSomething = true;
+	bHasRenderedSomething = true;
 	return InCombinedResult;
 }
 
 void ULandscapeTexturePatch::ApplyToWeightmap(bool bInPerformBlending, UE::Landscape::EditLayers::FRenderParams* RenderParams, UE::Landscape::FRDGBuilderRecorder& RDGBuilderRecorder, ULandscapeWeightPatchTextureInfo* PatchInfo,
 	FTextureResource* InMergedLandscapeTextureResource, int32 LandscapeTextureSliceIndex,
-	const FIntPoint& LandscapeTextureResolution, const FTransform& LandscapeHeightmapToWorld, bool& bOutDidRenderSomething, ERHIAccess OutputAccess)
+	const FIntPoint& LandscapeTextureResolution, const FTransform& LandscapeHeightmapToWorld, bool& bHasRenderedSomething, ERHIAccess OutputAccess)
 {
 	using namespace UE::Landscape;
-	bOutDidRenderSomething = false;
 
 	if (!PatchInfo)
 	{
@@ -569,7 +569,7 @@ void ULandscapeTexturePatch::ApplyToWeightmap(bool bInPerformBlending, UE::Lands
 		ULandscapeScratchRenderTarget* WriteRT = RenderParams->MergeRenderContext->GetBlendRenderTargetWrite();
 		// After this point, the render cannot fail so if we're the first in our render layer group to render, we can cycle the blend render targets and 
 		//  start rendering in the write one : 
-		if (RenderParams->NumSuccessfulRenderLayerStepsUntilBlendLayerStep == 0)
+		if (!bHasRenderedSomething && (RenderParams->NumSuccessfulRenderLayerStepsUntilBlendLayerStep == 0))
 		{
 			RenderParams->MergeRenderContext->CycleBlendRenderTargets(RDGBuilderRecorder);
 			WriteRT = RenderParams->MergeRenderContext->GetBlendRenderTargetWrite();
@@ -652,7 +652,7 @@ void ULandscapeTexturePatch::ApplyToWeightmap(bool bInPerformBlending, UE::Lands
 	// We need to specify the final state of the external texture to prevent the graph builder from transitioning it to SRVMask :
 	RDGBuilderRecorder.EnqueueRDGCommand(RDGCommand, { { InMergedLandscapeTextureResource, (OutputAccess == ERHIAccess::None) ? ERHIAccess::RTV : OutputAccess} });
 	
-	bOutDidRenderSomething = true;
+	bHasRenderedSomething = true;
 }
 
 void ULandscapeTexturePatch::GetCommonShaderParams(const FTransform& LandscapeHeightmapToWorldIn,
@@ -1097,6 +1097,10 @@ void ULandscapeTexturePatch::ReinitializeHeight(UTextureRenderTarget2D* InCombin
 		HeightInternalData->GetRenderTarget()->Modify();
 		HeightInternalData->CopyBackFromInternalTexture();
 	}
+
+	// Request a new landscape update to take into account the changes applied to the texture right away
+	//  Defer it till next frame (ExecuteOnGameThread) since requesting an update while updating won't do anything
+	ExecuteOnGameThread(TEXT("DeferredReinitializeHeightPatch"), [this] { RequestLandscapeUpdate(); });
 }
 
 void ULandscapeTexturePatch::ReinitializeWeightPatch(ULandscapeWeightPatchTextureInfo* PatchInfo, 
@@ -1194,6 +1198,10 @@ void ULandscapeTexturePatch::ReinitializeWeightPatch(ULandscapeWeightPatchTextur
 	});
 
 	PatchInfo->InternalData->SetUseInternalTextureOnly(PatchInfo->SourceMode == ELandscapeTexturePatchSourceMode::InternalTexture, true);
+
+	// Request a new landscape update to take into account the changes applied to the texture right away
+	//  Defer it till next frame (ExecuteOnGameThread) since requesting an update while updating won't do anything
+	ExecuteOnGameThread(TEXT("DeferredReinitializeWeightPatch"), [this] { RequestLandscapeUpdate(); });
 }
 
 FMatrix44f ULandscapeTexturePatch::GetPatchToHeightmapUVs(const FTransform& LandscapeHeightmapToWorld,
@@ -1277,6 +1285,20 @@ void ULandscapeTexturePatch::GetRenderDependencies(TSet<UObject*>& OutDependenci
 	{
 		OutDependencies.Add(HeightTextureAsset);
 	}
+
+	Algo::ForEach(WeightPatches, [&OutDependencies](const TObjectPtr<ULandscapeWeightPatchTextureInfo>& Info)
+		{ 
+			if (Info->SourceMode == ELandscapeTexturePatchSourceMode::InternalTexture
+				&& Info->InternalData && Info->InternalData->GetInternalTexture())
+			{
+				OutDependencies.Add(Info->InternalData->GetInternalTexture());
+			}
+			else if (Info->SourceMode == ELandscapeTexturePatchSourceMode::TextureAsset
+				&& Info->TextureAsset)
+			{
+				OutDependencies.Add(Info->TextureAsset);
+			}
+		});
 }
 
 TStructOnScope<FActorComponentInstanceData> ULandscapeTexturePatch::GetComponentInstanceData() const
@@ -2221,13 +2243,26 @@ TArray<UE::Landscape::EditLayers::FEditLayerRenderItem> ULandscapeTexturePatch::
 
 	FTransform ComponentTransform = this->GetComponentToWorld();
 	FOOBox2D PatchArea(ComponentTransform, GetUnscaledCoverage());
-	FInputWorldArea InputWorldArea = FInputWorldArea::CreateOOBox(PatchArea);
+	// The output is exactly the patch's area (i.e. object-oriented box)
 	FOutputWorldArea OutputWorldArea = FOutputWorldArea::CreateOOBox(PatchArea);
+	// Each pixel only depends on the pixel above so we don't need to read anything else than the component itself : 
+	FInputWorldArea InputWorldArea = FInputWorldArea::CreateLocalComponent();
+
+	// HACK [jonathan.bard] (the whole reinitialize height/weight is a hack currently anyway : this will disappear once this is implemented via a batched merge partial render) 
+	//  When reinitializing height/weight, we need to make sure the patch will be rendered in one batch and one only, because bReinitializeHeightOnNextRender will be reset 
+	//  upon rendering, so if 2 batches render the patch, it will be reset between the 2 renders, which will screw up the render command recorder. On top of it, we need to make
+	//  sure the entire patch area is rendered in one operation, because we read the result back into a render target
+	bool bForceSingleBatch = false;
 
 	if (InMergeContext->IsHeightmapMerge())
 	{
 		if (CanAffectHeightmap())
 		{
+			if (bReinitializeHeightOnNextRender)
+			{
+				bForceSingleBatch = true;
+			}
+
 			FEditLayerTargetTypeState TargetInfo(InMergeContext, ELandscapeToolTargetTypeFlags::Heightmap);
 			FEditLayerRenderItem Item(TargetInfo, InputWorldArea, OutputWorldArea, false);
 			AffectedAreas.Add(Item);
@@ -2239,6 +2274,11 @@ TArray<UE::Landscape::EditLayers::FEditLayerRenderItem> ULandscapeTexturePatch::
 		{
 			if (IsValid(WeightPatch) && WeightPatchCanRender(*WeightPatch))
 			{
+				if (WeightPatch->bReinitializeOnNextRender)
+				{
+					bForceSingleBatch = true;
+				}
+
 				FEditLayerTargetTypeState TargetInfo(InMergeContext);
 				if (WeightPatch->bEditVisibilityLayer)
 				{
@@ -2256,6 +2296,15 @@ TArray<UE::Landscape::EditLayers::FEditLayerRenderItem> ULandscapeTexturePatch::
 		}
 	}
 
+	// TODO [jonathan.bard] Remove once landscape edit layers partial merge replaces bReinitializeHeightOnNextRender / bReinitializeOnNextRender
+	if (bForceSingleBatch)
+	{
+		Algo::ForEach(AffectedAreas, [](FEditLayerRenderItem& InRenderItem) 
+			{
+				InRenderItem.SetInputWorldArea(FInputWorldArea::CreateInfinite());
+			});
+	}
+
 	return AffectedAreas;
 }
 #endif // WITH_EDITOR
diff --git a/Engine/Plugins/Experimental/LandscapePatch/Source/LandscapePatch/Public/LandscapeTexturePatch.h b/Engine/Plugins/Experimental/LandscapePatch/Source/LandscapePatch/Public/LandscapeTexturePatch.h
--- a/Engine/Plugins/Experimental/LandscapePatch/Source/LandscapePatch/Public/LandscapeTexturePatch.h
+++ b/Engine/Plugins/Experimental/LandscapePatch/Source/LandscapePatch/Public/LandscapeTexturePatch.h
@@ -662,11 +662,11 @@ private:
 	FLandscapeHeightPatchConvertToNativeParams GetHeightConvertToNativeParams() const;
 
 	UTextureRenderTarget2D* ApplyToHeightmap(bool bInPerformBlending, UE::Landscape::EditLayers::FRenderParams* RenderParams, UE::Landscape::FRDGBuilderRecorder& RDGBuilderRecorder, UTextureRenderTarget2D* InCombinedResult,
-		const FTransform& LandscapeHeightmapToWorld, bool& bOutDidRenderSomething, ERHIAccess OutputAccess = ERHIAccess::None);
+		const FTransform& LandscapeHeightmapToWorld, bool& bHasRenderedSomething, ERHIAccess OutputAccess = ERHIAccess::None);
 
 	void ApplyToWeightmap(bool bInPerformBlending, UE::Landscape::EditLayers::FRenderParams* RenderParams, UE::Landscape::FRDGBuilderRecorder& RDGBuilderRecorder, ULandscapeWeightPatchTextureInfo* PatchInfo,
 		FTextureResource* InMergedLandscapeTextureResource, int32 LandscapeTextureSliceIndex, const FIntPoint& LandscapeTextureResolution, 
-		const FTransform& LandscapeHeightmapToWorld, bool& bOutDidRenderSomething, ERHIAccess OutputAccess = ERHIAccess::None);
+		const FTransform& LandscapeHeightmapToWorld, bool& bHasRenderedSomething, ERHIAccess OutputAccess = ERHIAccess::None);
 
 	void GetCommonShaderParams(const FTransform& LandscapeHeightmapToWorldIn, 
 		const FIntPoint& SourceResolutionIn, const FIntPoint& DestinationResolutionIn,
diff --git a/Engine/Source/Runtime/Landscape/Private/LandscapeEditLayerMergeRenderContext.cpp b/Engine/Source/Runtime/Landscape/Private/LandscapeEditLayerMergeRenderContext.cpp
--- a/Engine/Source/Runtime/Landscape/Private/LandscapeEditLayerMergeRenderContext.cpp
+++ b/Engine/Source/Runtime/Landscape/Private/LandscapeEditLayerMergeRenderContext.cpp
@@ -28,7 +28,7 @@
 extern TAutoConsoleVariable<float> CVarLandscapeBatchedMergeVisualLogOffsetIncrement;
 extern TAutoConsoleVariable<int32> CVarLandscapeEditLayersClearBeforeEachWriteToScratch;
 extern TAutoConsoleVariable<int32> CVarLandscapeBatchedMergeVisualLogShowMergeType;
-extern TAutoConsoleVariable<bool> CVarLandscapeBatchedMergeVisualLogShowMergeProcess;
+extern TAutoConsoleVariable<int32> CVarLandscapeBatchedMergeVisualLogShowMergeProcess;
 extern TAutoConsoleVariable<float> CVarLandscapeBatchedMergeVisualLogAlpha;
 
 namespace UE::Landscape::EditLayers
@@ -815,7 +815,7 @@ void FMergeRenderContext::Render(TFunction<void(const FOnRenderBatchTargetGroupD
 
 	check(CurrentRenderBatchIndex == INDEX_NONE);
 
-	const bool bVisualLogMergeProcess = CVarLandscapeBatchedMergeVisualLogShowMergeProcess.GetValueOnGameThread();
+	const int32 ShowMergeProcess = CVarLandscapeBatchedMergeVisualLogShowMergeProcess.GetValueOnGameThread();
 
 	const FTransform& LandscapeTransform = Landscape->GetTransform();
 	// For visual logging, start at the top of the landscape's bounding box :
@@ -847,7 +847,7 @@ void FMergeRenderContext::Render(TFunction<void(const FOnRenderBatchTargetGroupD
 
 		// Drop a visual log showing the area covered by this batch : 
 		UE_IFVLOG(
-			if (IsVisualLogEnabled() && bVisualLogMergeProcess)
+			if (IsVisualLogEnabled() && ShowMergeProcess != 0)
 			{
 				// Pick a new color for each batch : 
 				uint32 Hash = PointerHash(&RenderBatch);
@@ -857,6 +857,13 @@ void FMergeRenderContext::Render(TFunction<void(const FOnRenderBatchTargetGroupD
 				UE_VLOG_OBOX(Landscape, LogLandscape, Log, FBox(FVector(RenderBatch.SectionRect.Min) - FVector(0.5, 0.5, 0.0), FVector(RenderBatch.SectionRect.Max) - FVector(0.5, 0.5, 0.0)),
 					ComputeVisualLogTransform(LandscapeWorldTransformForVisLog).ToMatrixWithScale(), Color.WithAlpha(GetVisualLogAlpha()),
 					TEXT("%s"), *FString::Format(TEXT("{0}\nBatch.SectionRect=([{1},{2}],[{3},{4}])"), { *RenderBatchDebugName, RenderBatch.SectionRect.Min.X, RenderBatch.SectionRect.Min.Y, RenderBatch.SectionRect.Max.X, RenderBatch.SectionRect.Max.Y }));
+
+				// Draw each component's bounds rendered by this renderer : 
+				for (const ULandscapeComponent* Component : RenderBatch.ComponentsToRender)
+				{
+					UE_VLOG_WIREOBOX(Landscape, LogLandscape, Log, FBox(FVector(Component->GetSectionBase()), FVector(Component->GetSectionBase() + Component->ComponentSizeQuads)),
+						ComputeVisualLogTransform(LandscapeWorldTransformForVisLog).ToMatrixWithScale(), FColor::White, TEXT(""));
+				}
 			});
 
 		const int32 NumRenderSteps = RenderBatch.RenderSteps.Num();
@@ -991,7 +998,7 @@ void FMergeRenderContext::Render(TFunction<void(const FOnRenderBatchTargetGroupD
 
 				// Drop some visual cues to help understand how each renderer is applied :
 				UE_IFVLOG(
-				if (IsVisualLogEnabled() && !bIsSeparateBlendStep && bVisualLogMergeProcess)
+				if (IsVisualLogEnabled() && !bIsSeparateBlendStep && (ShowMergeProcess == 2))
 					{
 						FTransform RenderAreaWorldTransformForVisLog = FTransform(FVector(RenderBatch.SectionRect.Min)) * LandscapeWorldTransformForVisLog;
 						IncrementVisualLogOffset();
diff --git a/Engine/Source/Runtime/Landscape/Private/LandscapeEditLayers.cpp b/Engine/Source/Runtime/Landscape/Private/LandscapeEditLayers.cpp
--- a/Engine/Source/Runtime/Landscape/Private/LandscapeEditLayers.cpp
+++ b/Engine/Source/Runtime/Landscape/Private/LandscapeEditLayers.cpp
@@ -226,10 +226,10 @@ TAutoConsoleVariable<int32> CVarLandscapeBatchedMergeVisualLogShowMergeType(
 	3,
 	TEXT("Filter what to visualize in the visual logger when using batched merge (0 = no visual log, 1 = show heightmaps only, 2 = show weightmaps only, 3 = show all"));
 
-TAutoConsoleVariable<bool> CVarLandscapeBatchedMergeVisualLogShowMergeProcess(
+TAutoConsoleVariable<int32> CVarLandscapeBatchedMergeVisualLogShowMergeProcess(
 	TEXT("landscape.BatchedMerge.VisualLog.ShowMergeProcess"),
-	false,
-	TEXT("Allows to visualize the merge process in the visual logger (batches, affected components, etc.)"));
+	0,
+	TEXT("Allows to visualize the merge process in the visual logger (0 = no visual log, 1 = show batches only, 2 = show batches and affected components per edit layer renderer)"));
 
 TAutoConsoleVariable<int32> CVarLandscapeBatchedMergeVisualLogShowRenderItemsType(
 	TEXT("landscape.BatchedMerge.VisualLog.ShowRenderItemsType"),
@@ -7186,7 +7186,9 @@ UE::Landscape::EditLayers::FMergeRenderContext ALandscape::PrepareEditLayersMerg
 							check((RenderLayerGroupStartIndex != INDEX_NONE) && (RenderLayerGroupEndIndex != INDEX_NONE));
 							const FEditLayerRendererRenderInfo& LastEditLayerRendererRenderInfo = OrderedEditLayerRendererRenderInfos[RenderLayerGroupEndIndex];
 							const ERenderFlags LastRenderFlags = LastEditLayerRendererRenderInfo.RendererState.GetRenderer()->GetRenderFlags(&MergeRenderContext);
+							const ERenderFlags LastRenderMode = LastRenderFlags & ERenderFlags::RenderMode_Mask;
 							TArrayView<ULandscapeComponent*> LastComponentsToRender = MakeArrayView(BatchEditLayerRendererRenderInfos[RenderLayerGroupEndIndex].ComponentsToRender);
+							check((LastRenderMode == ERenderFlags::RenderMode_Recorded) == bInRecordedSequence); // when closing a group, the last blend step should be in the same mode as the last render step
 							MergeRenderBatch.RenderSteps.Emplace(FMergeRenderStep::EType::EndRenderLayerGroup, LastRenderFlags, LastEditLayerRendererRenderInfo.RendererState, TargetLayerGroup, LastComponentsToRender);
 							MergeRenderBatch.RenderSteps.Emplace(FMergeRenderStep::EType::BlendLayer, LastRenderFlags, LastEditLayerRendererRenderInfo.RendererState, TargetLayerGroup, LastComponentsToRender);
 							RenderLayerGroupStartIndex = INDEX_NONE;
@@ -7195,22 +7197,6 @@ UE::Landscape::EditLayers::FMergeRenderContext ALandscape::PrepareEditLayersMerg
 
 					if (!ComponentsToRender.IsEmpty())
 					{
-						// Initiate the "render command recorder" sequence if necessary :
-						if (CurrentRenderMode == ERenderFlags::RenderMode_Recorded)
-						{
-							if (!bInRecordedSequence)
-							{
-								MergeRenderBatch.RenderSteps.Emplace(FMergeRenderStep::EType::BeginRenderCommandRecorder);
-								bInRecordedSequence = true;
-							}
-						}
-						// Terminate the "render command recorder" sequence if necessary :
-						else if (bInRecordedSequence)
-						{
-							MergeRenderBatch.RenderSteps.Emplace(FMergeRenderStep::EType::EndRenderCommandRecorder);
-							bInRecordedSequence = false;
-						}
-
 						if (bSupportsGrouping)
 						{
 							// Start a new render layer group if this renderer supports it and none is currently being built
@@ -7231,6 +7217,22 @@ UE::Landscape::EditLayers::FMergeRenderContext ALandscape::PrepareEditLayersMerg
 							AppendEndRenderLayerGroupAndBlendLayerSteps();
 						}
 
+						// Initiate the "render command recorder" sequence if necessary :
+						if (CurrentRenderMode == ERenderFlags::RenderMode_Recorded)
+						{
+							if (!bInRecordedSequence)
+							{
+								MergeRenderBatch.RenderSteps.Emplace(FMergeRenderStep::EType::BeginRenderCommandRecorder);
+								bInRecordedSequence = true;
+							}
+						}
+						// Or terminate the "render command recorder" sequence if necessary :
+						else if (bInRecordedSequence)
+						{
+							MergeRenderBatch.RenderSteps.Emplace(FMergeRenderStep::EType::EndRenderCommandRecorder);
+							bInRecordedSequence = false;
+						}
+
 						// Render the content of this layer : 
 						MergeRenderBatch.RenderSteps.Emplace(FMergeRenderStep::EType::RenderLayer, CurrentRenderFlags, EditLayerRendererRenderInfo.RendererState, TargetLayerGroup, ComponentsToRender);
 						if (bSupportsGrouping)
