Keep landscape patches in sync with their host landscape by upgrading them from the soft-dirty package list to regular dirty status when the landscape makes that same change.

diff --git a/Engine/Plugins/Experimental/LandscapePatch/Source/LandscapePatch/Private/LandscapePatchComponent.cpp b/Engine/Plugins/Experimental/LandscapePatch/Source/LandscapePatch/Private/LandscapePatchComponent.cpp
--- a/Engine/Plugins/Experimental/LandscapePatch/Source/LandscapePatch/Private/LandscapePatchComponent.cpp
+++ b/Engine/Plugins/Experimental/LandscapePatch/Source/LandscapePatch/Private/LandscapePatchComponent.cpp
@@ -721,7 +721,7 @@ void ULandscapePatchComponent::RequestLandscapeUpdate(bool bInUserTriggeredUpdat
 			if (!bGaveNotInPatchManagerWarning)
 			{
 				UE_LOG(LogLandscapePatch, Warning, TEXT("Patch's patch manager does not contain this patch. "
-					"Perhaps the manager was not saved? Reset the patch manager on the patch. (Package: %s, Actor: %s)"),
+					"Run LandscapePatch.FixPatchBindings or reset the patch manager on the patch. (Package: %s, Actor: %s)"),
 					*GetNameSafe(GetPackage()), *GetNameSafe(GetAttachmentRootActor()));
 				bGaveNotInPatchManagerWarning = true;
 			}
@@ -732,7 +732,7 @@ void ULandscapePatchComponent::RequestLandscapeUpdate(bool bInUserTriggeredUpdat
 			if (!bGaveMissingLandscapeWarning)
 			{
 				UE_LOG(LogLandscapePatch, Warning, TEXT("Patch's patch manager does not have a valid owning "
-					"landscape. Perhaps the landscape was not saved? Reset the landscape on the manager. (Package: %s, Manager: %s)"),
+					"landscape. Run LandscapePatch.FixPatchBindings or reset the landscape on the manager. (Package: %s, Manager: %s)"),
 					*GetNameSafe(GetPackage()), *GetNameSafe(PatchManager.Get()));
 				bGaveMissingLandscapeWarning = true;
 			}
@@ -775,7 +775,7 @@ void ULandscapePatchComponent::RequestLandscapeUpdate(bool bInUserTriggeredUpdat
 		// We didn't even have a guid
 		if (!bGaveMissingEditLayerGuidWarning)
 		{
-			UE_LOG(LogLandscapePatch, Warning, TEXT("Patch is not bound to an edit layer. Set the landscape and edit layer guid "
+			UE_LOG(LogLandscapePatch, Warning, TEXT("Patch is not bound to an edit layer. Set the landscape and edit layer guid or run LandscapePatch.FixPatchBindings."
 				"on the patch. (Package: %s)"),
 				*GetNameSafe(GetPackage()));
 			bGaveMissingEditLayerGuidWarning = true;
diff --git a/Engine/Source/Runtime/Landscape/Classes/LandscapeInfo.h b/Engine/Source/Runtime/Landscape/Classes/LandscapeInfo.h
--- a/Engine/Source/Runtime/Landscape/Classes/LandscapeInfo.h
+++ b/Engine/Source/Runtime/Landscape/Classes/LandscapeInfo.h
@@ -467,6 +467,8 @@ private:
 	bool UpdateLayerInfoMapInternal(ALandscapeProxy* Proxy);
 	bool TryAddToModifiedPackages(UPackage* InPackage, const ALandscape* InLandscapeOverride = nullptr);
 	void ProcessDeferredDeletions();
+
+	void OnMarkPackageDirty(UPackage* InPackage, bool bIsDirty);
 #endif
 
 private:
diff --git a/Engine/Source/Runtime/Landscape/Private/Landscape.cpp b/Engine/Source/Runtime/Landscape/Private/Landscape.cpp
--- a/Engine/Source/Runtime/Landscape/Private/Landscape.cpp
+++ b/Engine/Source/Runtime/Landscape/Private/Landscape.cpp
@@ -5216,6 +5216,26 @@ int32 ULandscapeInfo::MarkModifiedPackagesAsDirty()
 	return NumDirtied;
 }
 
+void ULandscapeInfo::OnMarkPackageDirty(UPackage* InPackage, bool bIsDirty)
+{
+	// Looking for the landscape actor's package to be marked dirty.
+	ALandscape* Landscape = LandscapeActor.Get();
+	if (!Landscape || Landscape->GetPackage() != InPackage)
+	{
+		return;
+	}
+
+	if (IsPackageModified(InPackage))
+	{
+		// LandscapeActor is on the soft-dirty list and the package is becoming actual-dirty. Also process the rest of the
+		// ModifiedPackages list.  Try to avoid saving the landscape out-of-sync with dependent objects (patches can have
+		// a problem with this).  At least provide the proper prompt to the user with all relevant actors on it.
+
+		ModifiedPackages.Remove(InPackage);  // Remove the main landscape package first to avoid infinitely looping.
+		MarkModifiedPackagesAsDirty();
+	}
+}
+
 int32 ULandscapeInfo::GetModifiedPackageCount() const
 {
 	return IntCastChecked<int32>(Algo::CountIf(ModifiedPackages, [](const TWeakObjectPtr<UPackage>& InWeakPackagePtr) { return InWeakPackagePtr.IsValid(); }));
@@ -5882,6 +5902,7 @@ void ULandscapeInfo::RegisterActor(ALandscapeProxy* Proxy, bool bMapCheck, bool
 			LandscapeActor = Landscape;
 
 #if WITH_EDITOR
+			UPackage::PackageMarkedDirtyEvent.AddWeakLambda(this, [this](UPackage* Pkg, bool bIsDirty) { OnMarkPackageDirty(Pkg, bIsDirty); });
 			// Now we have associated a LandscapeActor with this info
 			// we can ask for the WeightMaps
 			UpdateLayerInfoMap(LandscapeActor.Get());
@@ -6011,6 +6032,7 @@ void ULandscapeInfo::UnregisterActor(ALandscapeProxy* Proxy)
 		if (LandscapeActor.Get() == Landscape)
 		{
 			LandscapeActor = nullptr;
+			UPackage::PackageMarkedDirtyEvent.RemoveAll(this);
 		}
 
 		// update proxies reference to landscape actor
