Handle null landscape patch weight entries in the couple of places where it was not

[RN]

diff --git a/Engine/Plugins/Experimental/LandscapePatch/Source/LandscapePatch/Private/LandscapeTexturePatch.cpp b/Engine/Plugins/Experimental/LandscapePatch/Source/LandscapePatch/Private/LandscapeTexturePatch.cpp
--- a/Engine/Plugins/Experimental/LandscapePatch/Source/LandscapePatch/Private/LandscapeTexturePatch.cpp
+++ b/Engine/Plugins/Experimental/LandscapePatch/Source/LandscapePatch/Private/LandscapeTexturePatch.cpp
@@ -139,7 +139,10 @@ UE::Landscape::EditLayers::ERenderFlags ULandscapeTexturePatch::GetRenderFlags(c
 	{
 		RenderFlags |= ERenderFlags::RenderMode_Immediate;
 	}
-	else if (!InMergeContext->IsHeightmapMerge() && Algo::AnyOf(WeightPatches, [](const TObjectPtr<ULandscapeWeightPatchTextureInfo>& Info) { return Info->bReinitializeOnNextRender; }))
+	else if (!InMergeContext->IsHeightmapMerge() && Algo::AnyOf(WeightPatches, [](const TObjectPtr<ULandscapeWeightPatchTextureInfo>& InWeightPatch)
+		{
+			return IsValid(InWeightPatch) && InWeightPatch->bReinitializeOnNextRender;
+		}))
 	{
 		RenderFlags |= ERenderFlags::RenderMode_Immediate;
 	}
@@ -210,28 +213,31 @@ bool ULandscapeTexturePatch::RenderLayer(UE::Landscape::EditLayers::FRenderParam
 			&& UE::Landscape::IsVisibilityLayer(RenderParams.TargetLayerGroupLayerInfos[TargetLayerIndexInGroup]);
 		
 		// Try to find the weight patch
-		ULandscapeWeightPatchTextureInfo* WeightPatchInfo = nullptr;
-		for (const TObjectPtr<ULandscapeWeightPatchTextureInfo>& WeightPatchEntry : WeightPatches)
+		ULandscapeWeightPatchTextureInfo* FoundWeightPatch = nullptr;
+		for (const TObjectPtr<ULandscapeWeightPatchTextureInfo>& WeightPatch : WeightPatches)
 		{
-			if ((bIsVisibilityLayer && WeightPatchEntry->bEditVisibilityLayer)
-				|| (WeightPatchEntry->WeightmapLayerName == RenderParams.TargetLayerGroupLayerNames[TargetLayerIndexInGroup]))
+			if (IsValid(WeightPatch))
 			{
-				WeightPatchInfo = WeightPatchEntry;
-				break;
+				if ((bIsVisibilityLayer && WeightPatch->bEditVisibilityLayer)
+					|| (WeightPatch->WeightmapLayerName == RenderParams.TargetLayerGroupLayerNames[TargetLayerIndexInGroup]))
+				{
+					FoundWeightPatch = WeightPatch;
+					break;
+				}
 			}
 		}
 
-		if (!WeightPatchInfo)
+		if (!FoundWeightPatch)
 		{
 			// Didn't have a patch for this weight layer
 			continue;
 		}
 
-		if (WeightPatchInfo->bReinitializeOnNextRender)
+		if (FoundWeightPatch->bReinitializeOnNextRender)
 		{
-			WeightPatchInfo->bReinitializeOnNextRender = false;
+			FoundWeightPatch->bReinitializeOnNextRender = false;
 			checkf(!RDGBuilderRecorder.IsRecording(), TEXT("We should be using ERenderFlags::RenderMode_Immediate when reinitializing weight"));
-			ReinitializeWeightPatch(WeightPatchInfo, TextureArray->GetResource(),
+			ReinitializeWeightPatch(FoundWeightPatch, TextureArray->GetResource(),
 				FIntPoint(TextureArray->SizeX, TextureArray->SizeY), TargetLayerIndexInGroup, LandscapeHeightmapToWorld);
 
 			bHasRenderedSomething = true;
@@ -239,7 +245,7 @@ bool ULandscapeTexturePatch::RenderLayer(UE::Landscape::EditLayers::FRenderParam
 		else
 		{
 			checkf(RDGBuilderRecorder.IsRecording(), TEXT("We should be using ERenderFlags::RenderMode_Recorded in the typical case"));
-			ApplyToWeightmap(/*bInPerformBlending = */false, &RenderParams, RDGBuilderRecorder, WeightPatchInfo,
+			ApplyToWeightmap(/*bInPerformBlending = */false, &RenderParams, RDGBuilderRecorder, FoundWeightPatch,
 				TextureArray->GetResource(),
 				TargetLayerIndexInGroup,
 				RenderParams.RenderAreaSectionRect.Size(), 
@@ -309,34 +315,34 @@ UTextureRenderTarget2D* ULandscapeTexturePatch::RenderLayer_Native(const FLandsc
 	else
 	{
 		// Try to find the weight patch
-		ULandscapeWeightPatchTextureInfo* WeightPatchInfo = nullptr;
+		ULandscapeWeightPatchTextureInfo* FoundWeightPatch = nullptr;
 
-		for (const TObjectPtr<ULandscapeWeightPatchTextureInfo>& WeightPatchEntry : WeightPatches)
+		for (const TObjectPtr<ULandscapeWeightPatchTextureInfo>& WeightPatch : WeightPatches)
 		{
-			if (!IsValid(WeightPatchEntry))
+			if (!IsValid(WeightPatch))
 			{
 				continue;
 			}
 
-			if ((bIsWeightmapTarget && (WeightPatchEntry->WeightmapLayerName == InParameters.WeightmapLayerName)) ||
-				(bIsVisibilityLayerTarget && WeightPatchEntry->bEditVisibilityLayer))
+			if ((bIsWeightmapTarget && (WeightPatch->WeightmapLayerName == InParameters.WeightmapLayerName)) ||
+				(bIsVisibilityLayerTarget && WeightPatch->bEditVisibilityLayer))
 			{
-				WeightPatchInfo = WeightPatchEntry;
+				FoundWeightPatch = WeightPatch;
 				break;
 			}
 		}
 
-		if (!WeightPatchInfo)
+		if (!FoundWeightPatch)
 		{
 			return InParameters.CombinedResult;
 		}
 
-		if (WeightPatchInfo->bReinitializeOnNextRender)
+		if (FoundWeightPatch->bReinitializeOnNextRender)
 		{
-			WeightPatchInfo->bReinitializeOnNextRender = false;
+			FoundWeightPatch->bReinitializeOnNextRender = false;
 			if (ensure(InParameters.CombinedResult->GetResource()))
 			{
-				ReinitializeWeightPatch(WeightPatchInfo, InParameters.CombinedResult->GetResource(), 
+				ReinitializeWeightPatch(FoundWeightPatch, InParameters.CombinedResult->GetResource(),
 					FIntPoint(InParameters.CombinedResult->SizeX, InParameters.CombinedResult->SizeY), 
 					-1, // Signifies that this is not a Texture2DArray
 					LandscapeHeightmapToWorld);
@@ -348,7 +354,7 @@ UTextureRenderTarget2D* ULandscapeTexturePatch::RenderLayer_Native(const FLandsc
 			// We render in immediate mode for the legacy path so we use an immediate recorder :  
 			FRDGBuilderRecorder RDGBuilderRecorderImmediate;
 			bool bHasRenderedSomething = false;
-			ApplyToWeightmap(/*bInPerformBlending = */true, /*RenderParams = */nullptr, RDGBuilderRecorderImmediate, WeightPatchInfo,
+			ApplyToWeightmap(/*bInPerformBlending = */true, /*RenderParams = */nullptr, RDGBuilderRecorderImmediate, FoundWeightPatch,
 				InParameters.CombinedResult->GetResource(),
 				0, // Slice index
 				FIntPoint(InParameters.CombinedResult->SizeX, InParameters.CombinedResult->SizeY),
@@ -923,9 +929,9 @@ void ULandscapeTexturePatch::RequestReinitializeWeights()
 			TArray<TObjectPtr<ULandscapeWeightPatchTextureInfo>> FoundPatches;
 			if (bIsVisibilityLayer)
 			{
-				FoundPatches = WeightPatches.FilterByPredicate([](const TObjectPtr<ULandscapeWeightPatchTextureInfo>& PatchInfo)
+				FoundPatches = WeightPatches.FilterByPredicate([](const TObjectPtr<ULandscapeWeightPatchTextureInfo>& InWeightPatch)
 				{
-					return IsValid(PatchInfo) && PatchInfo->bEditVisibilityLayer;
+					return IsValid(InWeightPatch) && InWeightPatch->bEditVisibilityLayer;
 				});
 			}
 			else
@@ -935,9 +941,9 @@ void ULandscapeTexturePatch::RequestReinitializeWeights()
 					continue;
 				}
 				FoundPatches = WeightPatches.FilterByPredicate(
-					[&WeightmapLayerName](const TObjectPtr<ULandscapeWeightPatchTextureInfo>& PatchInfo) 
+					[&WeightmapLayerName](const TObjectPtr<ULandscapeWeightPatchTextureInfo>& InWeightPatch)
 				{ 
-					return PatchInfo && PatchInfo->WeightmapLayerName == WeightmapLayerName; 
+					return IsValid(InWeightPatch) && InWeightPatch->WeightmapLayerName == WeightmapLayerName;
 				});
 			}
 
@@ -949,9 +955,12 @@ void ULandscapeTexturePatch::RequestReinitializeWeights()
 			}
 			else
 			{
-				for (const TObjectPtr<ULandscapeWeightPatchTextureInfo>& PatchInfo : FoundPatches)
+				for (const TObjectPtr<ULandscapeWeightPatchTextureInfo>& WeightPatch : FoundPatches)
 				{
-					PatchInfo->bReinitializeOnNextRender = true;
+					if (IsValid(WeightPatch))
+					{
+						WeightPatch->bReinitializeOnNextRender = true;
+					}
 				}
 			}
 		}
@@ -1286,17 +1295,20 @@ void ULandscapeTexturePatch::GetRenderDependencies(TSet<UObject*>& OutDependenci
 		OutDependencies.Add(HeightTextureAsset);
 	}
 
-	Algo::ForEach(WeightPatches, [&OutDependencies](const TObjectPtr<ULandscapeWeightPatchTextureInfo>& Info)
+	Algo::ForEach(WeightPatches, [&OutDependencies](const TObjectPtr<ULandscapeWeightPatchTextureInfo>& InWeightPatch)
 		{ 
-			if (Info->SourceMode == ELandscapeTexturePatchSourceMode::InternalTexture
-				&& Info->InternalData && Info->InternalData->GetInternalTexture())
-			{
-				OutDependencies.Add(Info->InternalData->GetInternalTexture());
-			}
-			else if (Info->SourceMode == ELandscapeTexturePatchSourceMode::TextureAsset
-				&& Info->TextureAsset)
+			if (IsValid(InWeightPatch))
 			{
-				OutDependencies.Add(Info->TextureAsset);
+				if (InWeightPatch->SourceMode == ELandscapeTexturePatchSourceMode::InternalTexture
+					&& InWeightPatch->InternalData && InWeightPatch->InternalData->GetInternalTexture())
+				{
+					OutDependencies.Add(InWeightPatch->InternalData->GetInternalTexture());
+				}
+				else if (InWeightPatch->SourceMode == ELandscapeTexturePatchSourceMode::TextureAsset
+					&& InWeightPatch->TextureAsset)
+				{
+					OutDependencies.Add(InWeightPatch->TextureAsset);
+				}
 			}
 		});
 }
@@ -1970,10 +1982,10 @@ void ULandscapeTexturePatch::AddWeightPatch(const FName& WeightmapLayerName, ELa
 
 void ULandscapeTexturePatch::RemoveWeightPatch(const FName& InWeightmapLayerName)
 {
-	WeightPatches.RemoveAll([InWeightmapLayerName](const TObjectPtr<ULandscapeWeightPatchTextureInfo>& PatchInfo)
-	{ 
-		return PatchInfo && PatchInfo->WeightmapLayerName == InWeightmapLayerName; 
-	});
+	WeightPatches.RemoveAll([InWeightmapLayerName](const TObjectPtr<ULandscapeWeightPatchTextureInfo>& InWeightPatch)
+		{
+			return IsValid(InWeightPatch) && InWeightPatch->WeightmapLayerName == InWeightmapLayerName;
+		});
 }
 
 void ULandscapeTexturePatch::RemoveAllWeightPatches()
