Added API exposure to virtual functions on landscape patch plugin UCLASSes so that they can be inherited from in other modules

diff --git a/Engine/Plugins/Experimental/LandscapePatch/Source/LandscapePatch/Public/LandscapePatchComponent.h b/Engine/Plugins/Experimental/LandscapePatch/Source/LandscapePatch/Public/LandscapePatchComponent.h
--- a/Engine/Plugins/Experimental/LandscapePatch/Source/LandscapePatch/Public/LandscapePatchComponent.h
+++ b/Engine/Plugins/Experimental/LandscapePatch/Source/LandscapePatch/Public/LandscapePatchComponent.h
@@ -50,10 +50,10 @@ class ULandscapePatchComponent : public USceneComponent
 	GENERATED_BODY()
 
 public:
-	ULandscapePatchComponent(const FObjectInitializer& ObjectInitializer = FObjectInitializer::Get());
+	LANDSCAPEPATCH_API ULandscapePatchComponent(const FObjectInitializer& ObjectInitializer = FObjectInitializer::Get());
 
 	// Called from global merge path
-	virtual UTextureRenderTarget2D* RenderLayer_Native(const FLandscapeBrushParameters& InParameters, 
+	LANDSCAPEPATCH_API virtual UTextureRenderTarget2D* RenderLayer_Native(const FLandscapeBrushParameters& InParameters,
 		const FTransform& HeightmapToWorld) { return InParameters.CombinedResult; }
 
 #if WITH_EDITOR
@@ -63,10 +63,10 @@ public:
 
 	// These determine whether the patch is configured correctly to affect height/weightmaps,
 	// ignoring whether it is currently enabled or not.
-	virtual bool CanAffectHeightmap() const { return false; }
-	virtual bool CanAffectWeightmap() const { return false; }
-	virtual bool CanAffectWeightmapLayer(const FName& InLayerName) const { return false; }
-	virtual bool CanAffectVisibilityLayer() const { return false; }
+	LANDSCAPEPATCH_API virtual bool CanAffectHeightmap() const { return false; }
+	LANDSCAPEPATCH_API virtual bool CanAffectWeightmap() const { return false; }
+	LANDSCAPEPATCH_API virtual bool CanAffectWeightmapLayer(const FName& InLayerName) const { return false; }
+	LANDSCAPEPATCH_API virtual bool CanAffectVisibilityLayer() const { return false; }
 
 	bool CanAffectLandscape() const { return CanAffectHeightmap() || CanAffectWeightmap() || CanAffectVisibilityLayer(); }
 
@@ -77,7 +77,7 @@ public:
 	bool AffectsVisibilityLayer() const { return IsEnabled() && CanAffectVisibilityLayer(); }
 
 	// Allows the patch to specify textures that need to be ready/compiled before applying the patch
-	virtual void GetRenderDependencies(TSet<UObject*>& OutDependencies) const {}
+	LANDSCAPEPATCH_API virtual void GetRenderDependencies(TSet<UObject*>& OutDependencies) const {}
 
 	UFUNCTION(BlueprintCallable, Category = "LandscapePatch")
 	void RequestLandscapeUpdate(bool bInUserTriggeredUpdate = false);
@@ -93,7 +93,7 @@ public:
 	 * @return false if the patch is marked as disabled and therefore can't affect the landscape.
 	 */
 	UFUNCTION(BlueprintCallable, Category = "LandscapePatch")
-	virtual bool IsEnabled() const { return bIsEnabled; }
+	LANDSCAPEPATCH_API virtual bool IsEnabled() const { return bIsEnabled; }
 
 	UFUNCTION(BlueprintCallable, Category = "LandscapePatch")
 	double GetPriority() const { return Priority; }
@@ -114,15 +114,15 @@ public:
 	FTransform GetLandscapeHeightmapCoordsToWorld() const;
 
 	UFUNCTION(BlueprintCallable, Category = "LandscapePatch")
-	virtual void SetLandscape(ALandscape* NewLandscape);
+	LANDSCAPEPATCH_API virtual void SetLandscape(ALandscape* NewLandscape);
 
 	UE_DEPRECATED(5.5, "Patch manager use is deprecated. Patches should point to an edit layer using a Guid, and order themselves using priority.")
 	UFUNCTION(BlueprintCallable, Category = "LandscapePatch")
-	virtual void SetPatchManager(ALandscapePatchManager* NewPatchManager);
+	LANDSCAPEPATCH_API virtual void SetPatchManager(ALandscapePatchManager* NewPatchManager);
 
 	UE_DEPRECATED(5.5, "Patch manager use is deprecated. Patches should point to an edit layer using a Guid, and order themselves using priority.")
 	UFUNCTION(BlueprintCallable, Category = "LandscapePatch")
-	virtual ALandscapePatchManager* GetPatchManager() const;
+	LANDSCAPEPATCH_API virtual ALandscapePatchManager* GetPatchManager() const;
 
 	// Verifies that we're not an archetype or trash or in an invalid world
 	bool IsPatchInWorld() const;
@@ -140,21 +140,21 @@ public:
 	void NotifyOfBoundLayerDeletion(ULandscapePatchEditLayer* LayerIn);
 
 	// USceneComponent
-	virtual void OnUpdateTransform(EUpdateTransformFlags UpdateTransformFlags, ETeleportType Teleport) override;
+	LANDSCAPEPATCH_API virtual void OnUpdateTransform(EUpdateTransformFlags UpdateTransformFlags, ETeleportType Teleport) override;
 
 	// UActorComponent
-	virtual void CheckForErrors() override;
-	virtual void OnComponentCreated() override;
-	virtual void OnComponentDestroyed(bool bDestroyingHierarchy) override;
-	virtual void OnRegister() override;
-	virtual void GetActorDescProperties(FPropertyPairsMap& PropertyPairsMap) const override;
-	virtual TStructOnScope<FActorComponentInstanceData> GetComponentInstanceData() const override;
+	LANDSCAPEPATCH_API virtual void CheckForErrors() override;
+	LANDSCAPEPATCH_API virtual void OnComponentCreated() override;
+	LANDSCAPEPATCH_API virtual void OnComponentDestroyed(bool bDestroyingHierarchy) override;
+	LANDSCAPEPATCH_API virtual void OnRegister() override;
+	LANDSCAPEPATCH_API virtual void GetActorDescProperties(FPropertyPairsMap& PropertyPairsMap) const override;
+	LANDSCAPEPATCH_API virtual TStructOnScope<FActorComponentInstanceData> GetComponentInstanceData() const override;
 
 	// UObject
-	virtual void PostEditChangeProperty(FPropertyChangedEvent& PropertyChangedEvent) override;
-	virtual bool IsPostLoadThreadSafe() const override { return true; }
-	virtual void PostLoad() override;
-	virtual void PostEditUndo() override;
+	LANDSCAPEPATCH_API virtual void PostEditChangeProperty(FPropertyChangedEvent& PropertyChangedEvent) override;
+	LANDSCAPEPATCH_API virtual bool IsPostLoadThreadSafe() const override { return true; }
+	LANDSCAPEPATCH_API virtual void PostLoad() override;
+	LANDSCAPEPATCH_API virtual void PostEditUndo() override;
 #endif
 	virtual bool IsEditorOnly() const override { return true; }
 	virtual bool NeedsLoadForClient() const override { return false; }
diff --git a/Engine/Plugins/Experimental/LandscapePatch/Source/LandscapePatch/Public/LandscapeTexturePatch.h b/Engine/Plugins/Experimental/LandscapePatch/Source/LandscapePatch/Public/LandscapeTexturePatch.h
--- a/Engine/Plugins/Experimental/LandscapePatch/Source/LandscapePatch/Public/LandscapeTexturePatch.h
+++ b/Engine/Plugins/Experimental/LandscapePatch/Source/LandscapePatch/Public/LandscapeTexturePatch.h
@@ -163,8 +163,8 @@ class ULandscapeWeightPatchTextureInfo : public UObject
 public:
 #if WITH_EDITOR
 	// UObject
-	virtual void PostEditChangeProperty(FPropertyChangedEvent& PropertyChangedEvent) override;
-	virtual void PreDuplicate(FObjectDuplicationParameters& DupParams) override;
+	LANDSCAPEPATCH_API virtual void PostEditChangeProperty(FPropertyChangedEvent& PropertyChangedEvent) override;
+	LANDSCAPEPATCH_API virtual void PreDuplicate(FObjectDuplicationParameters& DupParams) override;
 #endif // WITH_EDITOR
 
 protected:
@@ -232,33 +232,33 @@ public:
 	using FEditLayerRenderItem = UE::Landscape::EditLayers::FEditLayerRenderItem;
 
 	// ILandscapeEditLayerRenderer, via ULandscapePatchComponent
-	virtual void GetRendererStateInfo(const UE::Landscape::EditLayers::FMergeContext* InMergeContext,
+	LANDSCAPEPATCH_API virtual void GetRendererStateInfo(const UE::Landscape::EditLayers::FMergeContext* InMergeContext,
 		FEditLayerTargetTypeState& OutSupportedTargetTypeState,
 		FEditLayerTargetTypeState& OutEnabledTargetTypeState,
 		TArray<TBitArray<>>& OutTargetLayerGroups) const override;
-	virtual FString GetEditLayerRendererDebugName() const override;
-	virtual TArray<FEditLayerRenderItem> GetRenderItems(const UE::Landscape::EditLayers::FMergeContext* InMergeContext) const override;
-	virtual UE::Landscape::EditLayers::ERenderFlags GetRenderFlags(const UE::Landscape::EditLayers::FMergeContext* InMergeContext) const override;
-	virtual bool CanGroupRenderLayerWith(TScriptInterface<ILandscapeEditLayerRenderer> InOtherRenderer) const override;
-	virtual bool RenderLayer(UE::Landscape::EditLayers::FRenderParams& RenderParams, UE::Landscape::FRDGBuilderRecorder& RDGBuilderRecorder) override;
-	virtual void BlendLayer(UE::Landscape::EditLayers::FRenderParams& RenderParams, UE::Landscape::FRDGBuilderRecorder& RDGBuilderRecorder) override;
+	LANDSCAPEPATCH_API virtual FString GetEditLayerRendererDebugName() const override;
+	LANDSCAPEPATCH_API virtual TArray<FEditLayerRenderItem> GetRenderItems(const UE::Landscape::EditLayers::FMergeContext* InMergeContext) const override;
+	LANDSCAPEPATCH_API virtual UE::Landscape::EditLayers::ERenderFlags GetRenderFlags(const UE::Landscape::EditLayers::FMergeContext* InMergeContext) const override;
+	LANDSCAPEPATCH_API virtual bool CanGroupRenderLayerWith(TScriptInterface<ILandscapeEditLayerRenderer> InOtherRenderer) const override;
+	LANDSCAPEPATCH_API virtual bool RenderLayer(UE::Landscape::EditLayers::FRenderParams& RenderParams, UE::Landscape::FRDGBuilderRecorder& RDGBuilderRecorder) override;
+	LANDSCAPEPATCH_API virtual void BlendLayer(UE::Landscape::EditLayers::FRenderParams& RenderParams, UE::Landscape::FRDGBuilderRecorder& RDGBuilderRecorder) override;
 
 	// ULandscapePatchComponent
-	UTextureRenderTarget2D* RenderLayer_Native(const FLandscapeBrushParameters& InParameters, const FTransform& HeightmapToWorld) override;
-	virtual bool CanAffectHeightmap() const override;
-	virtual bool CanAffectWeightmap() const override;
-	virtual bool CanAffectWeightmapLayer(const FName& InLayerName) const override;
-	virtual bool CanAffectVisibilityLayer() const override;
-	virtual void GetRenderDependencies(TSet<UObject*>& OutDependencies) const override;
+	LANDSCAPEPATCH_API virtual UTextureRenderTarget2D* RenderLayer_Native(const FLandscapeBrushParameters& InParameters, const FTransform& HeightmapToWorld) override;
+	LANDSCAPEPATCH_API virtual bool CanAffectHeightmap() const override;
+	LANDSCAPEPATCH_API virtual bool CanAffectWeightmap() const override;
+	LANDSCAPEPATCH_API virtual bool CanAffectWeightmapLayer(const FName& InLayerName) const override;
+	LANDSCAPEPATCH_API virtual bool CanAffectVisibilityLayer() const override;
+	LANDSCAPEPATCH_API virtual void GetRenderDependencies(TSet<UObject*>& OutDependencies) const override;
 
 	// UActorComponent
-	virtual TStructOnScope<FActorComponentInstanceData> GetComponentInstanceData() const override;
-	virtual void CheckForErrors() override;
+	LANDSCAPEPATCH_API virtual TStructOnScope<FActorComponentInstanceData> GetComponentInstanceData() const override;
+	LANDSCAPEPATCH_API virtual void CheckForErrors() override;
 
 	// UObject
-	virtual void PostEditChangeProperty(FPropertyChangedEvent& PropertyChangedEvent) override;
+	LANDSCAPEPATCH_API virtual void PostEditChangeProperty(FPropertyChangedEvent& PropertyChangedEvent) override;
 #endif
-	virtual void Serialize(FArchive& Ar) override;
+	LANDSCAPEPATCH_API virtual void Serialize(FArchive& Ar) override;
 
 	/**
 	 * Gets the transform from patch to world. The transform is based off of the component
@@ -266,7 +266,7 @@ public:
 	 * to rotate it relative to the landscape.
 	 */
 	UFUNCTION(BlueprintCallable, Category = LandscapePatch)
-	virtual FTransform GetPatchToWorldTransform() const;
+	LANDSCAPEPATCH_API virtual FTransform GetPatchToWorldTransform() const;
 
 	/**
 	 * Gives size in unscaled world coordinates (ie before applying patch transform) of the patch as measured 
@@ -290,7 +290,7 @@ public:
 	 * the height/weight patches).
 	 */
 	UFUNCTION(BlueprintCallable, Category = LandscapePatch)
-	virtual FVector2D GetFullUnscaledWorldSize() const;
+	LANDSCAPEPATCH_API virtual FVector2D GetFullUnscaledWorldSize() const;
 
 	/**
 	 * Gets the size (in pixels) of the internal textures used by the patch. Does not reflect the resolution
@@ -306,7 +306,7 @@ public:
 	 * @return true if successful.
 	 */
 	UFUNCTION(BlueprintCallable, Category = LandscapePatch)
-	virtual void SetResolution(FVector2D ResolutionIn);
+	LANDSCAPEPATCH_API virtual void SetResolution(FVector2D ResolutionIn);
 
 	/**
 	 * Given the landscape resolution, current patch coverage, and a landscape resolution multiplier, gives the
@@ -316,7 +316,7 @@ public:
 	 * @return true if successful (may fail if landscape is not set, for instance)
 	 */
 	UFUNCTION(BlueprintCallable, Category = LandscapePatch, meta = (ResolutionMultiplier = "1.0"))
-	virtual UPARAM(DisplayName = "Success") bool GetInitResolutionFromLandscape(float ResolutionMultiplier, FVector2D& ResolutionOut) const;
+	LANDSCAPEPATCH_API virtual UPARAM(DisplayName = "Success") bool GetInitResolutionFromLandscape(float ResolutionMultiplier, FVector2D& ResolutionOut) const;
 
 
 	UFUNCTION(BlueprintCallable, Category = LandscapePatch)
@@ -363,7 +363,7 @@ public:
 	 * render target may not be initialized.
 	 */
 	UFUNCTION(BlueprintCallable, Category = LandscapePatch)
-	virtual void SetHeightSourceMode(ELandscapeTexturePatchSourceMode NewMode);
+	LANDSCAPEPATCH_API virtual void SetHeightSourceMode(ELandscapeTexturePatchSourceMode NewMode);
 
 	/**
 	 * Sets the texture used for height when the height source mode is set to texture asset. Note that
@@ -392,7 +392,7 @@ public:
 	 *  being written to. Can be set to false if the render target is not being written to.
 	 */
 	UFUNCTION(BlueprintCallable, Category = LandscapePatch)
-	virtual UTextureRenderTarget2D* GetHeightRenderTarget(bool bMarkDirty = true);
+	LANDSCAPEPATCH_API virtual UTextureRenderTarget2D* GetHeightRenderTarget(bool bMarkDirty = true);
 
 	UFUNCTION(BlueprintCallable, Category = LandscapePatch, meta = (ETextureRenderTargetFormat = "ETextureRenderTargetFormat::RTF_R32f"))
 	void SetHeightRenderTargetFormat(ETextureRenderTargetFormat Format);
@@ -461,29 +461,29 @@ public:
 	 * By default, the layer is added with source mode set to be a texture-backed render target.
 	 */
 	UFUNCTION(BlueprintCallable, Category = LandscapePatch)
-	virtual void AddWeightPatch(const FName& InWeightmapLayerName, ELandscapeTexturePatchSourceMode SourceMode, bool bUseAlphaChannel);
+	LANDSCAPEPATCH_API virtual void AddWeightPatch(const FName& InWeightmapLayerName, ELandscapeTexturePatchSourceMode SourceMode, bool bUseAlphaChannel);
 
 	UFUNCTION(BlueprintCallable, Category = LandscapePatch)
-	virtual void RemoveWeightPatch(const FName& InWeightmapLayerName);
+	LANDSCAPEPATCH_API virtual void RemoveWeightPatch(const FName& InWeightmapLayerName);
 
 	UFUNCTION(BlueprintCallable, Category = LandscapePatch)
-	virtual void RemoveAllWeightPatches();
+	LANDSCAPEPATCH_API virtual void RemoveAllWeightPatches();
 
 	/** Sets the source mode of all weight patches to "None". */
 	UFUNCTION(BlueprintCallable, Category = LandscapePatch)
-	virtual void DisableAllWeightPatches();
+	LANDSCAPEPATCH_API virtual void DisableAllWeightPatches();
 
 	UFUNCTION(BlueprintCallable, Category = LandscapePatch)
 	TArray<FName> GetAllWeightPatchLayerNames();
 	
 	UFUNCTION(BlueprintCallable, Category = LandscapePatch)
-	virtual ELandscapeTexturePatchSourceMode GetWeightPatchSourceMode(const FName& InWeightmapLayerName);
+	LANDSCAPEPATCH_API virtual ELandscapeTexturePatchSourceMode GetWeightPatchSourceMode(const FName& InWeightmapLayerName);
 
 	UFUNCTION(BlueprintCallable, Category = LandscapePatch)
-	virtual UTexture* GetWeightPatchTextureAsset(const FName& InWeightmapLayerName);
+	LANDSCAPEPATCH_API virtual UTexture* GetWeightPatchTextureAsset(const FName& InWeightmapLayerName);
 
 	UFUNCTION(BlueprintCallable, Category = LandscapePatch)
-	virtual void SetWeightPatchSourceMode(const FName& InWeightmapLayerName, ELandscapeTexturePatchSourceMode NewMode);
+	LANDSCAPEPATCH_API virtual void SetWeightPatchSourceMode(const FName& InWeightmapLayerName, ELandscapeTexturePatchSourceMode NewMode);
 
 	UFUNCTION(BlueprintCallable, Category = LandscapePatch)
 	void SetWeightPatchTextureAsset(const FName& InWeightmapLayerName, UTexture* TextureIn);
@@ -493,19 +493,19 @@ public:
 	 *  being written to. Can be set to false if the render target is not being written to.
 	 */
 	UFUNCTION(BlueprintCallable, Category = LandscapePatch)
-	virtual UTextureRenderTarget2D* GetWeightPatchRenderTarget(const FName& InWeightmapLayerName, bool bMarkDirty = true);
+	LANDSCAPEPATCH_API virtual UTextureRenderTarget2D* GetWeightPatchRenderTarget(const FName& InWeightmapLayerName, bool bMarkDirty = true);
 
 	UFUNCTION(BlueprintCallable, Category = LandscapePatch)
-	virtual void SetUseAlphaChannelForWeightPatch(const FName& InWeightmapLayerName, bool bUseAlphaChannel);
+	LANDSCAPEPATCH_API virtual void SetUseAlphaChannelForWeightPatch(const FName& InWeightmapLayerName, bool bUseAlphaChannel);
 
 	UFUNCTION(BlueprintCallable, Category = LandscapePatch)
-	virtual void SetWeightPatchBlendModeOverride(const FName& InWeightmapLayerName, ELandscapeTexturePatchBlendMode BlendMode);
+	LANDSCAPEPATCH_API virtual void SetWeightPatchBlendModeOverride(const FName& InWeightmapLayerName, ELandscapeTexturePatchBlendMode BlendMode);
 
 	UFUNCTION(BlueprintCallable, Category = LandscapePatch)
-	virtual void ClearWeightPatchBlendModeOverride(const FName& InWeightmapLayerName);
+	LANDSCAPEPATCH_API virtual void ClearWeightPatchBlendModeOverride(const FName& InWeightmapLayerName);
 
 	UFUNCTION(BlueprintCallable, Category = LandscapePatch)
-	virtual void SetEditVisibilityLayer(const FName& InWeightmapLayerName, const bool bEditVisibilityLayer);
+	LANDSCAPEPATCH_API virtual void SetEditVisibilityLayer(const FName& InWeightmapLayerName, const bool bEditVisibilityLayer);
 
 protected:
 	//~ Don't expose these on the instance because a user might not realize that they would lose their existing internal
