Fixed additive blend mode for landscape texture patches that got broken with batched merge

diff --git a/Engine/Plugins/Experimental/LandscapePatch/Shaders/Private/LandscapeTexturePatchPS.usf b/Engine/Plugins/Experimental/LandscapePatch/Shaders/Private/LandscapeTexturePatchPS.usf
--- a/Engine/Plugins/Experimental/LandscapePatch/Shaders/Private/LandscapeTexturePatchPS.usf
+++ b/Engine/Plugins/Experimental/LandscapePatch/Shaders/Private/LandscapeTexturePatchPS.usf
@@ -172,7 +172,7 @@ void ApplyLandscapeTextureHeightPatch(in float4 SVPos : SV_POSITION, out float4
 
 #else // PERFORM_BLENDING
 
-	float OutputHeight = (InBlendMode == ADDITIVE_MODE) ? PatchSignedHeight : (LANDSCAPE_MID_VALUE + PatchSignedHeight);
+	float OutputHeight = LANDSCAPE_MID_VALUE + PatchSignedHeight;
 	OutColor = float4(PackHeight(OutputHeight), PackHeightAlpha(Alpha, ConvertPatchBlendModeToHeightAlphaFlags()));
 
 #endif // !PERFORM_BLENDING
