Landscape Editor - Adding RequestLandscapeUpdate call into BlueprintBrushBase on PostEditUndo
* If a BlueprintBrushBase is deleted, then UNDO to revert the deletion, the landscape needs to be updated
* If not, ULandscapeComponent::ClearUpdateFlagsForModes will not be cleared properly in the next UpdateLayersContent

diff --git a/Engine/Plugins/Experimental/LandscapePatch/Source/LandscapePatch/Private/LandscapePatchManager.cpp b/Engine/Plugins/Experimental/LandscapePatch/Source/LandscapePatch/Private/LandscapePatchManager.cpp
--- a/Engine/Plugins/Experimental/LandscapePatch/Source/LandscapePatch/Private/LandscapePatchManager.cpp
+++ b/Engine/Plugins/Experimental/LandscapePatch/Source/LandscapePatch/Private/LandscapePatchManager.cpp
@@ -663,11 +663,6 @@ void ALandscapePatchManager::GetRenderDependencies(TSet<UObject*>& OutDependenci
 	}
 }
 
-void ALandscapePatchManager::PostEditUndo()
-{
-	RequestLandscapeUpdate();
-}
-
 void ALandscapePatchManager::SetOwningLandscape(ALandscape* InOwningLandscape)
 {
 	Super::SetOwningLandscape(InOwningLandscape);
diff --git a/Engine/Plugins/Experimental/LandscapePatch/Source/LandscapePatch/Public/LandscapePatchManager.h b/Engine/Plugins/Experimental/LandscapePatch/Source/LandscapePatch/Public/LandscapePatchManager.h
--- a/Engine/Plugins/Experimental/LandscapePatch/Source/LandscapePatch/Public/LandscapePatchManager.h
+++ b/Engine/Plugins/Experimental/LandscapePatch/Source/LandscapePatch/Public/LandscapePatchManager.h
@@ -123,7 +123,6 @@ public:
 	virtual void PostRegisterAllComponents() override;
 
 	// UObject
-	virtual void PostEditUndo() override;
 	virtual void PostEditChangeProperty(FPropertyChangedEvent& PropertyChangedEvent) override;
 	virtual void PostLoad() override;
 #endif
diff --git a/Engine/Source/Runtime/Landscape/Private/LandscapeBlueprintBrushBase.cpp b/Engine/Source/Runtime/Landscape/Private/LandscapeBlueprintBrushBase.cpp
--- a/Engine/Source/Runtime/Landscape/Private/LandscapeBlueprintBrushBase.cpp
+++ b/Engine/Source/Runtime/Landscape/Private/LandscapeBlueprintBrushBase.cpp
@@ -623,6 +623,11 @@ void ALandscapeBlueprintBrushBase::PostEditChangeProperty(FPropertyChangedEvent&
 	}
 }
 
+void ALandscapeBlueprintBrushBase::PostEditUndo()
+{
+	RequestLandscapeUpdate();
+}
+
 void ALandscapeBlueprintBrushBase::Destroyed()
 {
 	Super::Destroyed();
diff --git a/Engine/Source/Runtime/Landscape/Public/LandscapeBlueprintBrushBase.h b/Engine/Source/Runtime/Landscape/Public/LandscapeBlueprintBrushBase.h
--- a/Engine/Source/Runtime/Landscape/Public/LandscapeBlueprintBrushBase.h
+++ b/Engine/Source/Runtime/Landscape/Public/LandscapeBlueprintBrushBase.h
@@ -156,6 +156,7 @@ public:
 	LANDSCAPE_API virtual void Tick(float DeltaSeconds) override;
 	LANDSCAPE_API virtual void PostEditMove(bool bFinished) override;
 	LANDSCAPE_API virtual void PostEditChangeProperty(FPropertyChangedEvent& PropertyChangedEvent) override;
+	LANDSCAPE_API virtual void PostEditUndo() override;
 	LANDSCAPE_API virtual void Destroyed() override;
 
 	LANDSCAPE_API virtual void PushDeferredLayersContentUpdate();
