Landscape Editor - Adding dll exposure for public BlueprintCallable landscape patch functions

diff --git a/Engine/Plugins/Experimental/LandscapePatch/Source/LandscapePatch/Public/LandscapePatchComponent.h b/Engine/Plugins/Experimental/LandscapePatch/Source/LandscapePatch/Public/LandscapePatchComponent.h
--- a/Engine/Plugins/Experimental/LandscapePatch/Source/LandscapePatch/Public/LandscapePatchComponent.h
+++ b/Engine/Plugins/Experimental/LandscapePatch/Source/LandscapePatch/Public/LandscapePatchComponent.h
@@ -80,26 +80,26 @@ public:
 	virtual void GetRenderDependencies(TSet<UObject*>& OutDependencies) const {}
 
 	UFUNCTION(BlueprintCallable, Category = "LandscapePatch")
-	void RequestLandscapeUpdate(bool bInUserTriggeredUpdate = false);
+	LANDSCAPEPATCH_API void RequestLandscapeUpdate(bool bInUserTriggeredUpdate = false);
 
 	/**
 	 * Allows the patch to be disabled, so that it no longer affects the landscape. This can be useful
 	 * when deleting the patch is undesirable, usually when the disabling is temporary.
 	 */
 	UFUNCTION(BlueprintCallable, Category = "LandscapePatch")
-	void SetIsEnabled(bool bEnabledIn);
+	LANDSCAPEPATCH_API void SetIsEnabled(bool bEnabledIn);
 
 	/**
 	 * @return false if the patch is marked as disabled and therefore can't affect the landscape.
 	 */
 	UFUNCTION(BlueprintCallable, Category = "LandscapePatch")
-	virtual bool IsEnabled() const { return bIsEnabled; }
+	LANDSCAPEPATCH_API virtual bool IsEnabled() const { return bIsEnabled; }
 
 	UFUNCTION(BlueprintCallable, Category = "LandscapePatch")
-	double GetPriority() const { return Priority; }
+	LANDSCAPEPATCH_API double GetPriority() const { return Priority; }
 
 	UFUNCTION(BlueprintCallable, Category = "LandscapePatch")
-	void SetPriority(double PriorityIn);
+	LANDSCAPEPATCH_API void SetPriority(double PriorityIn);
 
 	FGuid GetEditLayerGuid() const { return EditLayerGuid; }
 	void SetEditLayerGuid(const FGuid& GuidIn);
@@ -111,7 +111,7 @@ public:
 	uint32 GetFullNameHash() const;
 
 	UFUNCTION(BlueprintCallable, Category = "LandscapePatch")
-	FTransform GetLandscapeHeightmapCoordsToWorld() const;
+	LANDSCAPEPATCH_API FTransform GetLandscapeHeightmapCoordsToWorld() const;
 
 	UFUNCTION(BlueprintCallable, Category = "LandscapePatch")
 	LANDSCAPEPATCH_API virtual void SetLandscape(ALandscape* NewLandscape);
diff --git a/Engine/Plugins/Experimental/LandscapePatch/Source/LandscapePatch/Public/LandscapeTexturePatch.h b/Engine/Plugins/Experimental/LandscapePatch/Source/LandscapePatch/Public/LandscapeTexturePatch.h
--- a/Engine/Plugins/Experimental/LandscapePatch/Source/LandscapePatch/Public/LandscapeTexturePatch.h
+++ b/Engine/Plugins/Experimental/LandscapePatch/Source/LandscapePatch/Public/LandscapeTexturePatch.h
@@ -275,13 +275,13 @@ public:
 	 * the landscape.
 	 */
 	UFUNCTION(BlueprintCallable, Category = LandscapePatch)
-	virtual FVector2D GetUnscaledCoverage() const { return FVector2D(UnscaledPatchCoverage); }
+	LANDSCAPEPATCH_API virtual FVector2D GetUnscaledCoverage() const { return FVector2D(UnscaledPatchCoverage); }
 
 	/**
 	 * Set the patch coverage (see GetUnscaledCoverage for description).
 	 */
 	UFUNCTION(BlueprintCallable, Category = LandscapePatch)
-	virtual void SetUnscaledCoverage(FVector2D Coverage) { UnscaledPatchCoverage = Coverage; }
+	LANDSCAPEPATCH_API virtual void SetUnscaledCoverage(FVector2D Coverage) { UnscaledPatchCoverage = Coverage; }
 
 	/**
 	 * When using an internal texture, gives size in unscaled world coordinates of the patch in the world,
@@ -297,7 +297,7 @@ public:
 	 * of any used texture assets (if the source mode is texture asset).
 	 */
 	UFUNCTION(BlueprintCallable, Category = LandscapePatch)
-	virtual FVector2D GetResolution() const { return FVector2D(ResolutionX, ResolutionY); }
+	LANDSCAPEPATCH_API virtual FVector2D GetResolution() const { return FVector2D(ResolutionX, ResolutionY); }
 
 	/**
 	 * Sets the resolution of the currently used internal texture or render target. Has no effect
@@ -320,7 +320,7 @@ public:
 
 
 	UFUNCTION(BlueprintCallable, Category = LandscapePatch)
-	void SetFalloff(float FalloffIn) 
+	LANDSCAPEPATCH_API void SetFalloff(float FalloffIn) 
 	{
 		if (Falloff != FalloffIn)
 		{
@@ -330,7 +330,7 @@ public:
 	}
 
 	UFUNCTION(BlueprintCallable, Category = LandscapePatch)
-	void SetFalloffMode(ELandscapeTexturePatchFalloffMode FalloffModeIn) 
+	LANDSCAPEPATCH_API void SetFalloffMode(ELandscapeTexturePatchFalloffMode FalloffModeIn) 
 	{
 		if (FalloffMode != FalloffModeIn)
 		{
@@ -343,7 +343,7 @@ public:
 	 * Determines how the height patch is blended into the existing terrain.
 	 */
 	UFUNCTION(BlueprintCallable, Category = LandscapePatch)
-	void SetBlendMode(ELandscapeTexturePatchBlendMode BlendModeIn) 
+	LANDSCAPEPATCH_API void SetBlendMode(ELandscapeTexturePatchBlendMode BlendModeIn) 
 	{
 		if (BlendMode != BlendModeIn)
 		{
@@ -355,7 +355,7 @@ public:
 	// Height related functions:
 
 	UFUNCTION(BlueprintCallable, Category = LandscapePatch)
-	virtual ELandscapeTexturePatchSourceMode GetHeightSourceMode() const { return HeightSourceMode; }
+	LANDSCAPEPATCH_API virtual ELandscapeTexturePatchSourceMode GetHeightSourceMode() const { return HeightSourceMode; }
 
 	/**
 	 * Changes source mode. There are currently no API guarantees regarding the initialization of the
@@ -370,7 +370,7 @@ public:
 	 * virtual textures are not supported.
 	 */
 	UFUNCTION(BlueprintCallable, Category = LandscapePatch)
-	void SetHeightTextureAsset(UTexture* TextureIn);
+	LANDSCAPEPATCH_API void SetHeightTextureAsset(UTexture* TextureIn);
 
 	/**
 	 * Gets the internal height render target, if source mode is set to Texture Backed Render Target.
@@ -395,7 +395,7 @@ public:
 	LANDSCAPEPATCH_API virtual UTextureRenderTarget2D* GetHeightRenderTarget(bool bMarkDirty = true);
 
 	UFUNCTION(BlueprintCallable, Category = LandscapePatch, meta = (ETextureRenderTargetFormat = "ETextureRenderTargetFormat::RTF_R32f"))
-	void SetHeightRenderTargetFormat(ETextureRenderTargetFormat Format);
+	LANDSCAPEPATCH_API void SetHeightRenderTargetFormat(ETextureRenderTargetFormat Format);
 
 	/**
 	 * Determines whether the height patch alpha channel is used for blending into the existing values.
@@ -403,7 +403,7 @@ public:
 	 * used depends on the patch blend mode (see SetBlendMode).
 	 */
 	UFUNCTION(BlueprintCallable, Category = "LandscapePatch")
-	void SetUseAlphaChannelForHeight(bool bUse)
+	LANDSCAPEPATCH_API void SetUseAlphaChannelForHeight(bool bUse)
 	{
 		if (bUseTextureAlphaForHeight != bUse)
 		{
@@ -417,7 +417,7 @@ public:
 	 * are translated into heights when applying to landscape.
 	 */
 	UFUNCTION(BlueprintCallable, Category = LandscapePatch)
-	void SetHeightEncodingMode(ELandscapeTextureHeightPatchEncoding EncodingMode)
+	LANDSCAPEPATCH_API void SetHeightEncodingMode(ELandscapeTextureHeightPatchEncoding EncodingMode)
 	{
 		if (HeightEncoding != EncodingMode)
 		{
@@ -431,7 +431,7 @@ public:
 	 * render target format to mode-specific defaults.
 	 */
 	UFUNCTION(BlueprintCallable, Category = LandscapePatch)
-	void ResetHeightEncodingMode(ELandscapeTextureHeightPatchEncoding EncodingMode);
+	LANDSCAPEPATCH_API void ResetHeightEncodingMode(ELandscapeTextureHeightPatchEncoding EncodingMode);
 
 	/**
 	 * Set settings that determine how values in the patch are translated into heights. This is only
@@ -439,13 +439,13 @@ public:
 	 * in the same space as the landscape heightmap.
 	 */
 	UFUNCTION(BlueprintCallable, Category = LandscapePatch)
-	void SetHeightEncodingSettings(const FLandscapeTexturePatchEncodingSettings& Settings);
+	LANDSCAPEPATCH_API void SetHeightEncodingSettings(const FLandscapeTexturePatchEncodingSettings& Settings);
 
 	/**
 	 * Set how zero height is interpreted, see comments in ELandscapeTextureHeightPatchZeroHeightMeaning.
 	 */
 	UFUNCTION(BlueprintCallable, Category = "LandscapePatch")
-	void SetZeroHeightMeaning(ELandscapeTextureHeightPatchZeroHeightMeaning ZeroHeightMeaningIn)
+	LANDSCAPEPATCH_API void SetZeroHeightMeaning(ELandscapeTextureHeightPatchZeroHeightMeaning ZeroHeightMeaningIn)
 	{ 
 		if (ZeroHeightMeaning != ZeroHeightMeaningIn)
 		{
@@ -474,7 +474,7 @@ public:
 	LANDSCAPEPATCH_API virtual void DisableAllWeightPatches();
 
 	UFUNCTION(BlueprintCallable, Category = LandscapePatch)
-	TArray<FName> GetAllWeightPatchLayerNames();
+	LANDSCAPEPATCH_API TArray<FName> GetAllWeightPatchLayerNames();
 	
 	UFUNCTION(BlueprintCallable, Category = LandscapePatch)
 	LANDSCAPEPATCH_API virtual ELandscapeTexturePatchSourceMode GetWeightPatchSourceMode(const FName& InWeightmapLayerName);
@@ -486,7 +486,7 @@ public:
 	LANDSCAPEPATCH_API virtual void SetWeightPatchSourceMode(const FName& InWeightmapLayerName, ELandscapeTexturePatchSourceMode NewMode);
 
 	UFUNCTION(BlueprintCallable, Category = LandscapePatch)
-	void SetWeightPatchTextureAsset(const FName& InWeightmapLayerName, UTexture* TextureIn);
+	LANDSCAPEPATCH_API void SetWeightPatchTextureAsset(const FName& InWeightmapLayerName, UTexture* TextureIn);
 
 	/**
 	 * @param bMarkDirty If true, marks the containing package as dirty, since the render target is presumably
@@ -651,11 +651,11 @@ protected:
 	UPROPERTY(EditAnywhere, Category = Initialization, meta = (EditCondition = "!bBaseResolutionOffLandscape", ClampMin = "1"))
 	int32 InitTextureSizeY = 33;
 
-private:
 	UTexture2D* GetHeightInternalTexture();
 	UTextureRenderTarget2D* GetWeightPatchRenderTarget(ULandscapeWeightPatchTextureInfo* WeightPatch);
 	UTexture2D* GetWeightPatchInternalTexture(ULandscapeWeightPatchTextureInfo* WeightPatch);
 
+private:
 	void UpdateHeightConvertToNativeParamsIfNeeded();
 #if WITH_EDITOR
 	void TransitionHeightSourceModeInternal(ELandscapeTexturePatchSourceMode OldMode, ELandscapeTexturePatchSourceMode NewMode);
