Don't visualize disabled landscape patches

diff --git a/Engine/Plugins/Experimental/LandscapePatch/Source/LandscapePatchEditorOnly/Private/LandscapeTexturePatchVisualizer.cpp b/Engine/Plugins/Experimental/LandscapePatch/Source/LandscapePatchEditorOnly/Private/LandscapeTexturePatchVisualizer.cpp
--- a/Engine/Plugins/Experimental/LandscapePatch/Source/LandscapePatchEditorOnly/Private/LandscapeTexturePatchVisualizer.cpp
+++ b/Engine/Plugins/Experimental/LandscapePatch/Source/LandscapePatchEditorOnly/Private/LandscapeTexturePatchVisualizer.cpp
@@ -13,7 +13,7 @@ void FLandscapeTexturePatchVisualizer::DrawVisualization(const UActorComponent*
 	float DepthBias = 1;
 	bool bScreenSpace = 1;
 
-	if (const ULandscapeTexturePatch* Patch = Cast<ULandscapeTexturePatch>(Component))
+	if (const ULandscapeTexturePatch* Patch = Cast<ULandscapeTexturePatch>(Component); (Patch != nullptr) && Patch->IsEnabled())
 	{
 		FTransform PatchToWorld = Patch->GetPatchToWorldTransform();
 		DrawRectangle(PDI, PatchToWorld.GetTranslation(), PatchToWorld.GetUnitAxis(EAxis::X), PatchToWorld.GetUnitAxis(EAxis::Y),
