Landscape : removed global and local merge code and deprecated anything that was only available for this.
Since the console variable "landscape.EditLayersLocalMerge.Enable" is deprecated now, a new console command was added, to trigger a landscape layer full update (landscape.ForceLayersFullUpdate) :
* With no argument : will update all landscapes
* With arguments : will update all landscapes whose display name (partially) matches the argument (e.g. with 3 landscapes named Landscape_0 Landscape_1 and Landscape_2, "landscape.ForceLayersFullUpdate _0 _1" will trigger an update on the first 2 only)

diff --git a/Engine/Plugins/Experimental/LandscapePatch/Shaders/Private/LandscapeCircleHeightPatchPS.usf b/Engine/Plugins/Experimental/LandscapePatch/Shaders/Private/LandscapeCircleHeightPatchPS.usf
--- a/Engine/Plugins/Experimental/LandscapePatch/Shaders/Private/LandscapeCircleHeightPatchPS.usf
+++ b/Engine/Plugins/Experimental/LandscapePatch/Shaders/Private/LandscapeCircleHeightPatchPS.usf
@@ -9,9 +9,6 @@ float3 InCenter;
 float InRadius;
 float InFalloff;
 
-Texture2D<float4> InSourceTexture;
-uint2 InSourceTextureOffset;
-
 #if CIRCLE_HEIGHT_PATCH
 
 void ApplyLandscapeCircleHeightPatch(in float4 SVPos : SV_POSITION, out float4 OutColor : SV_Target0)
@@ -21,12 +18,6 @@ void ApplyLandscapeCircleHeightPatch(in float4 SVPos : SV_POSITION, out float4 O
 	
 	int2 DestinationTextureCoordinates = floor(SVPos.xy);
 
-#if PERFORM_BLENDING
-	int2 SourceTextureCoordinates = DestinationTextureCoordinates - InSourceTextureOffset;
-	float4 CurrentPackedHeight = InSourceTexture.Load(int3(SourceTextureCoordinates, 0));
-	float CurrentHeight = UnpackHeight(CurrentPackedHeight.xy);
-#endif // PERFORM_BLENDING
-	
 	float PatchHeight = InCenter.z;
 	
 	float Distance = distance(DestinationTextureCoordinates, InCenter.xy);
@@ -37,15 +28,7 @@ void ApplyLandscapeCircleHeightPatch(in float4 SVPos : SV_POSITION, out float4 O
 		Alpha = 1 - min((Distance - InRadius) / ClampedFalloff, 1.0);
 	}
 	
-// TODO [jonathan.bard] Deprecate :
-#if PERFORM_BLENDING
-	float NewHeight = lerp(CurrentHeight, PatchHeight, Alpha);
-	OutColor = float4(PackHeight(NewHeight), 0.0, 0.0);
-#else // PERFORM_BLENDING
-
 	OutColor = float4(PackHeight(PatchHeight), PackHeightAlpha(Alpha, EHEIGHTMAPALPHAFLAGS_ALPHABLEND));
-
-#endif // !PERFORM_BLENDING
 }
 
 #endif // CIRCLE_HEIGHT_PATCH
@@ -57,27 +40,7 @@ void ApplyLandscapeCircleVisibilityPatch(in float4 SVPos : SV_POSITION, out floa
 	int2 DestinationTextureCoordinates = floor(SVPos.xy);
 	float Distance = distance(DestinationTextureCoordinates, InCenter.xy);
 	
-// TODO [jonathan.bard] Deprecate :
-#if PERFORM_BLENDING
-	int2 SourceTextureCoordinates = DestinationTextureCoordinates - InSourceTextureOffset;
-	float4 CurrentWeight = InSourceTexture.Load(int3(SourceTextureCoordinates, 0)).x;
-
-	if (Distance <= InRadius)
-	{
-		// TODO: This should be changed to calculate alpha just like a height patch and write out
-		//  the alpha value, which allows for smoother holes. But needs to be changed in a
-		//  separate CL with a fixup to set the falloff of older assets to 0 to keep them unchanged.
-		OutColor = 1.0;
-	}
-	else
-	{
-		OutColor = CurrentWeight.x;
-	}
-#else // PERFORM_BLENDING
-
 	OutColor = (Distance <= InRadius) ? 1.0f : 0.0f;
-
-#endif // !PERFORM_BLENDING
 }
 
 #endif // CIRCLE_VISIBILITY_PATCH
diff --git a/Engine/Plugins/Experimental/LandscapePatch/Shaders/Private/LandscapeTexturePatchPS.usf b/Engine/Plugins/Experimental/LandscapePatch/Shaders/Private/LandscapeTexturePatchPS.usf
--- a/Engine/Plugins/Experimental/LandscapePatch/Shaders/Private/LandscapeTexturePatchPS.usf
+++ b/Engine/Plugins/Experimental/LandscapePatch/Shaders/Private/LandscapeTexturePatchPS.usf
@@ -92,8 +92,6 @@ float4 SamplePatchTexture(Texture2D<float4> InPatchTexture, SamplerState InPatch
 #endif // APPLY_HEIGHT_PATCH || APPLY_WEIGHT_PATCH
 
 #if APPLY_HEIGHT_PATCH
-Texture2D<float4> InSourceHeightmap;
-uint2 InSourceHeightmapOffset;
 Texture2D<float4> InHeightPatch;
 SamplerState InHeightPatchSampler;
 float4x4 InHeightmapToPatch;
@@ -147,12 +145,6 @@ void ApplyLandscapeTextureHeightPatch(in float4 SVPos : SV_POSITION, out float4
 	
 	float PatchSignedHeight = InHeightScale * (PatchStoredHeight - InZeroInEncoding) + InHeightOffset;
 	
-#if PERFORM_BLENDING
-	int2 HeightmapCoordinates = floor(SVPos.xy) - InSourceHeightmapOffset;
-	float4 CurrentPackedHeight = InSourceHeightmap.Load(int3(HeightmapCoordinates, 0));
-	float CurrentHeight = UnpackHeight(CurrentPackedHeight.xy);
-#endif // PERFORM_BLENDING
-	
 	float Alpha = GetFalloffAlpha(InFalloffWorldMargin, InPatchWorldDimensions, PatchUVCoordinates, InEdgeUVDeadBorder, bRectangularFalloff);
 	
 	if (bApplyPatchAlpha)
@@ -160,33 +152,8 @@ void ApplyLandscapeTextureHeightPatch(in float4 SVPos : SV_POSITION, out float4
 		Alpha *= PatchSampledValue.a;
 	}
 	
-	// TODO [jonathan.bard] Deprecate :
-#if PERFORM_BLENDING
-	float NewHeight = 0;
-	switch (InBlendMode)
-	{
-		case ADDITIVE_MODE:
-			NewHeight = CurrentHeight + Alpha * PatchSignedHeight;
-			break;
-		case ALPHA_BLEND_MODE:
-			NewHeight = lerp(CurrentHeight, LANDSCAPE_MID_VALUE + PatchSignedHeight, Alpha);
-			break;
-		case MIN_MODE:
-			NewHeight = lerp(CurrentHeight, min(CurrentHeight, LANDSCAPE_MID_VALUE + PatchSignedHeight), Alpha);
-			break;
-		case MAX_MODE:
-			NewHeight = lerp(CurrentHeight, max(CurrentHeight, LANDSCAPE_MID_VALUE + PatchSignedHeight), Alpha);
-			break;
-	}
-	
-	OutColor = float4(PackHeight(NewHeight), PackHeightAlpha(1.0f, EHEIGHTMAPALPHAFLAGS_NONE));
-
-#else // PERFORM_BLENDING
-
 	float OutputHeight = LANDSCAPE_MID_VALUE + PatchSignedHeight;
 	OutColor = float4(PackHeight(OutputHeight), PackHeightAlpha(Alpha, ConvertPatchBlendModeToHeightAlphaFlags()));
-
-#endif // !PERFORM_BLENDING
 }
 #endif // APPLY_HEIGHT_PATCH
 
@@ -282,11 +249,6 @@ void ApplyLandscapeTextureWeightPatch(in float4 SVPos : SV_POSITION, out float4
 	float4 PatchSampledValue = SamplePatchTexture(InWeightPatch, InWeightPatchSampler, PatchUVCoordinates);
 	float PatchWeight = PatchSampledValue.x;
 	
-#if PERFORM_BLENDING
-	int2 SourceWeightmapCoordinates = floor(SVPos.xy) - InSourceWeightmapCoordOffset;
-	float CurrentWeight = InSourceWeightmap.Load(int3(SourceWeightmapCoordinates, 0)).x;
-#endif // PERFORM_BLENDING
-	
 	float Alpha = GetFalloffAlpha(InFalloffWorldMargin, InPatchWorldDimensions, PatchUVCoordinates, InEdgeUVDeadBorder, bRectangularFalloff);
 	
 	if (bApplyPatchAlpha)
@@ -294,31 +256,7 @@ void ApplyLandscapeTextureWeightPatch(in float4 SVPos : SV_POSITION, out float4
 		Alpha *= PatchSampledValue.a;
 	}
 	
-	// TODO [jonathan.bard] Deprecate :
-#if PERFORM_BLENDING
-	float NewWeight = 0;
-	switch (InBlendMode)
-	{
-		case ADDITIVE_MODE:
-			NewWeight = clamp(CurrentWeight + Alpha * PatchWeight, 0, 1);
-			break;
-		case ALPHA_BLEND_MODE:
-			NewWeight = lerp(CurrentWeight, PatchWeight, Alpha);
-			break;
-		case MIN_MODE:
-			NewWeight = lerp(CurrentWeight, min(CurrentWeight, PatchWeight), Alpha);
-			break;
-		case MAX_MODE:
-			NewWeight = lerp(CurrentWeight, max(CurrentWeight, PatchWeight), Alpha);
-			break;
-	}
-	
-	OutColor = float4(PackWeight(NewWeight), PackWeightAlpha(1.0f, EWEIGHTMAPALPHAFLAGS_NONE));
-#else // PERFORM_BLENDING
-
 	OutColor = float4(PackWeight(PatchWeight), PackWeightAlpha(Alpha, ConvertPatchBlendModeToWeightAlphaFlags()));
-
-#endif // !PERFORM_BLENDING
 }
 
 
diff --git a/Engine/Plugins/Experimental/LandscapePatch/Source/LandscapePatch/Private/LandscapeCircleHeightPatch.cpp b/Engine/Plugins/Experimental/LandscapePatch/Source/LandscapePatch/Private/LandscapeCircleHeightPatch.cpp
--- a/Engine/Plugins/Experimental/LandscapePatch/Source/LandscapePatch/Private/LandscapeCircleHeightPatch.cpp
+++ b/Engine/Plugins/Experimental/LandscapePatch/Source/LandscapePatch/Private/LandscapeCircleHeightPatch.cpp
@@ -48,26 +48,6 @@ void ULandscapeCircleHeightPatch::OnComponentCreated()
 }
 
 #if WITH_EDITOR
-UTextureRenderTarget2D* ULandscapeCircleHeightPatch::RenderLayer_Native(const FLandscapeBrushParameters& InParameters, const FTransform& HeightmapCoordsToWorld)
-{
-	using namespace UE::Landscape;
-
-	// Circle height patch doesn't affect regular weightmap layers.
-	if ((bEditVisibility && (InParameters.LayerType != ELandscapeToolTargetType::Visibility)) || (!bEditVisibility && (InParameters.LayerType != ELandscapeToolTargetType::Heightmap)))
-	{
-		return InParameters.CombinedResult;
-	}
-
-	// We render in immediate mode for the legacy path so we use an immediate recorder :  
-	FRDGBuilderRecorder RDGBuilderRecorderImmediate;
-	ApplyCirclePatch(/*bInPerformBlending = */true, /*RenderParams = */nullptr, RDGBuilderRecorderImmediate, InParameters.LayerType == ELandscapeToolTargetType::Visibility,
-		InParameters.CombinedResult->GetResource(), 0,
-		FIntPoint(InParameters.CombinedResult->SizeX, InParameters.CombinedResult->SizeY),
-		HeightmapCoordsToWorld, /*OutputAccess = */ERHIAccess::SRVMask); // in the legacy path the RT leaves as a SRV
-
-	return InParameters.CombinedResult;
-}
-
 UE::Landscape::EditLayers::ERenderFlags ULandscapeCircleHeightPatch::GetRenderFlags(const UE::Landscape::EditLayers::FMergeContext* InMergeContext) const
 {
 	using namespace UE::Landscape::EditLayers;
@@ -98,7 +78,7 @@ bool ULandscapeCircleHeightPatch::RenderLayer(UE::Landscape::EditLayers::FRender
 	{
 		UTextureRenderTarget2D* OutputToBlendInto = LandscapeRT->TryGetRenderTarget2D();
 		check(OutputToBlendInto != nullptr);
-		return ApplyCirclePatch(/*bInPerformBlending = */false, &RenderParams, RDGBuilderRecorder, /*bIsVisibilityLayer = */false, OutputToBlendInto->GetResource(),
+		return ApplyCirclePatch(&RenderParams, RDGBuilderRecorder, /*bIsVisibilityLayer = */false, OutputToBlendInto->GetResource(),
 			0, RenderParams.RenderAreaSectionRect.Size(), HeightmapCoordsToWorld);
 	}
 
@@ -120,7 +100,7 @@ bool ULandscapeCircleHeightPatch::RenderLayer(UE::Landscape::EditLayers::FRender
 	{
 		if (UE::Landscape::IsVisibilityLayer(RenderParams.TargetLayerGroupLayerInfos[LayerIndex]))
 		{
-			bDidRenderSomething |= ApplyCirclePatch(/*bInPerformBlending = */false, &RenderParams, RDGBuilderRecorder, /*bIsVisibilityLayer = */true, TextureArray->GetResource(),
+			bDidRenderSomething |= ApplyCirclePatch(&RenderParams, RDGBuilderRecorder, /*bIsVisibilityLayer = */true, TextureArray->GetResource(),
 				LayerIndex, RenderParams.RenderAreaSectionRect.Size(), HeightmapCoordsToWorld);
 		}
 	}
@@ -146,7 +126,7 @@ void ULandscapeCircleHeightPatch::BlendLayer(UE::Landscape::EditLayers::FRenderP
 	RenderParams.MergeRenderContext->GenericBlendLayer(BlendParams, RenderParams, RDGBuilderRecorder);
 }
 
-bool ULandscapeCircleHeightPatch::ApplyCirclePatch(bool bInPerformBlending, UE::Landscape::EditLayers::FRenderParams* RenderParams, UE::Landscape::FRDGBuilderRecorder& RDGBuilderRecorder, bool bIsVisibilityLayer,
+bool ULandscapeCircleHeightPatch::ApplyCirclePatch(UE::Landscape::EditLayers::FRenderParams* RenderParams, UE::Landscape::FRDGBuilderRecorder& RDGBuilderRecorder, bool bIsVisibilityLayer,
 	FTextureResource* InMergedLandscapeTextureResource, int32 LandscapeTextureSliceIndex,
 	const FIntPoint& DestinationResolution, const FTransform& HeightmapCoordsToWorld, ERHIAccess OutputAccess)
 {
@@ -184,8 +164,6 @@ bool ULandscapeCircleHeightPatch::ApplyCirclePatch(bool bInPerformBlending, UE::
 
 	FTextureResource* OutputResource = InMergedLandscapeTextureResource;
 	FString OutputResourceName = OutputResource->GetResourceName().ToString();
-	if (!bInPerformBlending)
-	{ 
 		check(RenderParams != nullptr);
 		ULandscapeScratchRenderTarget* WriteRT = RenderParams->MergeRenderContext->GetBlendRenderTargetWrite();
 	    // After this point, the render cannot fail so if we're the first in our render layer group to render, we can cycle the blend render targets and 
@@ -199,7 +177,6 @@ bool ULandscapeCircleHeightPatch::ApplyCirclePatch(bool bInPerformBlending, UE::
 	    }
 		OutputResource = WriteRT->GetRenderTarget()->GetResource();
 		OutputResourceName = WriteRT->GetDebugName();
-	}
 
 	auto RDGCommand =
 		[ OutputResource
@@ -209,76 +186,28 @@ bool ULandscapeCircleHeightPatch::ApplyCirclePatch(bool bInPerformBlending, UE::
 		, CenterInHeightmapCoordinates
 		, HeightmapRadius
 		, HeightmapFalloff
-		, bEditVisibility = bEditVisibility
-		, bInPerformBlending](FRDGBuilder& GraphBuilder)
+		, bEditVisibility = bEditVisibility](FRDGBuilder& GraphBuilder)
 		{
 			TRACE_CPUPROFILER_EVENT_SCOPE(LandscapeCircleHeightPatch);
 
 			FRDGTextureRef DestinationTexture = GraphBuilder.RegisterExternalTexture(CreateRenderTarget(OutputResource->TextureRHI, TEXT("OutputTexture")));
 			const TCHAR* OutputName = bEditVisibility ? TEXT("LandscapeCircleVisibilityPatchOutput") : TEXT("LandscapeCircleHeightPatchOutput");
 
-			FRDGTextureSRVRef InputCopySRV = nullptr;
-			if (bInPerformBlending)
-			{
-			    // Make a copy of the portion of our weightmap input that we're writing to so that we can 
-			    //  read and write at the same time (needed for blending)
-			    FRDGTextureDesc InputCopyDescription = DestinationTexture->Desc;
-			    InputCopyDescription.Dimension = ETextureDimension::Texture2D;
-			    InputCopyDescription.ArraySize = 1;
-			    InputCopyDescription.NumMips = 1;
-			    InputCopyDescription.Extent = DestinationBounds.Size();
-			    const TCHAR* InputCopyName = bEditVisibility ? TEXT("LandscapeCircleVisibilityPatchInputCopy") : TEXT("LandscapeCircleHeightPatchInputCopy");
-			    FRDGTextureRef InputCopy = GraphBuilder.CreateTexture(InputCopyDescription, InputCopyName);
-    
-			    FRHICopyTextureInfo CopyTextureInfo;
-			    CopyTextureInfo.SourceMipIndex = 0;
-			    CopyTextureInfo.NumMips = 1;
-			    CopyTextureInfo.SourceSliceIndex = LandscapeTextureSliceIndex;
-			    CopyTextureInfo.NumSlices = 1;
-			    CopyTextureInfo.SourcePosition = FIntVector(DestinationBounds.Min.X, DestinationBounds.Min.Y, 0);
-			    CopyTextureInfo.Size = FIntVector(InputCopyDescription.Extent.X, InputCopyDescription.Extent.Y, 0);
-			    AddCopyTexturePass(GraphBuilder, DestinationTexture, InputCopy, CopyTextureInfo);
-    
-			    InputCopySRV = GraphBuilder.CreateSRV(FRDGTextureSRVDesc::CreateForMipLevel(InputCopy, 0));
-			}
-
-			FLandscapeCircleHeightPatchPSBase::FParameters* ShaderParams = GraphBuilder.AllocParameters<FLandscapeCircleHeightPatchPSBase::FParameters>();
+			FLandscapeCircleHeightPatchPS::FParameters* ShaderParams = GraphBuilder.AllocParameters<FLandscapeCircleHeightPatchPS::FParameters>();
 			ShaderParams->InCenter = (FVector3f)CenterInHeightmapCoordinates;
 			ShaderParams->InRadius = HeightmapRadius;
 			ShaderParams->InFalloff = HeightmapFalloff;
 			
-			ShaderParams->InSourceTexture = GraphBuilder.CreateSRV(FRDGTextureSRVDesc::Create(GSystemTextures.GetBlackDummy(GraphBuilder)));
-			ShaderParams->InSourceTextureOffset = FIntPoint(ForceInit);
-			if (bInPerformBlending)
-			{
-				ShaderParams->InSourceTexture = InputCopySRV;
-				ShaderParams->InSourceTextureOffset = DestinationBounds.Min;
-			}
-
 			ShaderParams->RenderTargets[0] = FRenderTargetBinding(DestinationTexture, ERenderTargetLoadAction::ENoAction,
 				/*InMipIndex = */0, LandscapeTextureSliceIndex);
 
 			if (bEditVisibility)
 			{
-				if (bInPerformBlending)
-				{
-					FLandscapeCircleVisibilityPatchPS</*bPerformBlending = */true>::AddToRenderGraph(RDG_EVENT_NAME("RenderCircleVisibilityPatch -> %s", *OutputResourceName), GraphBuilder, ShaderParams, DestinationBounds);
-				}
-				else
-				{
-					FLandscapeCircleVisibilityPatchPS</*bPerformBlending = */false>::AddToRenderGraph(RDG_EVENT_NAME("RenderCircleVisibilityPatch -> %s", *OutputResourceName), GraphBuilder, ShaderParams, DestinationBounds);
-				}
+				FLandscapeCircleVisibilityPatchPS::AddToRenderGraph(RDG_EVENT_NAME("RenderCircleVisibilityPatch -> %s", *OutputResourceName), GraphBuilder, ShaderParams, DestinationBounds);
 			}
 			else
 			{
-				if (bInPerformBlending)
-				{
-					FLandscapeCircleHeightPatchPS</*bPerformBlending = */true>::AddToRenderGraph(RDG_EVENT_NAME("RenderCircleHeightPatch -> %s", *OutputResourceName), GraphBuilder, ShaderParams, DestinationBounds);
-				}
-				else
-				{
-					FLandscapeCircleHeightPatchPS</*bPerformBlending = */false>::AddToRenderGraph(RDG_EVENT_NAME("RenderCircleHeightPatch -> %s", *OutputResourceName), GraphBuilder, ShaderParams, DestinationBounds);
-				}
+				FLandscapeCircleHeightPatchPS::AddToRenderGraph(RDG_EVENT_NAME("RenderCircleHeightPatch -> %s", *OutputResourceName), GraphBuilder, ShaderParams, DestinationBounds);
 			}
 		};
 
diff --git a/Engine/Plugins/Experimental/LandscapePatch/Source/LandscapePatch/Private/LandscapeCircleHeightPatchPS.h b/Engine/Plugins/Experimental/LandscapePatch/Source/LandscapePatch/Private/LandscapeCircleHeightPatchPS.h
--- a/Engine/Plugins/Experimental/LandscapePatch/Source/LandscapePatch/Private/LandscapeCircleHeightPatchPS.h
+++ b/Engine/Plugins/Experimental/LandscapePatch/Source/LandscapePatch/Private/LandscapeCircleHeightPatchPS.h
@@ -8,53 +8,39 @@
 /**
  * Shader that applies a circle height patch to a landscape heightmap.
  */
-class FLandscapeCircleHeightPatchPSBase : public FGlobalShader
+class FLandscapeCircleHeightPatchPS : public FGlobalShader
 {
+	DECLARE_GLOBAL_SHADER(FLandscapeCircleHeightPatchPS);
+	SHADER_USE_PARAMETER_STRUCT(FLandscapeCircleHeightPatchPS, FGlobalShader);
+
 public:
 	BEGIN_SHADER_PARAMETER_STRUCT(FParameters, )
-		SHADER_PARAMETER_RDG_TEXTURE_SRV(Texture2D<float4>, InSourceTexture) // Our input texture
-		// Offset of the source heightmap relative to the 0,0 location in the destination heightmap, because
-		// the source is likely to be a copied region from some inner part of the destination. This is basically
-		// a SourceHeightmapToDestinationHeightmap coordinate transformation, except that it is always a simple 
-		// integer translation.
-		SHADER_PARAMETER(FIntPoint, InSourceTextureOffset)
 		SHADER_PARAMETER(FVector3f, InCenter)
 		SHADER_PARAMETER(float, InRadius)
 		SHADER_PARAMETER(float, InFalloff)
 		RENDER_TARGET_BINDING_SLOTS() // Holds our output
 	END_SHADER_PARAMETER_STRUCT()
 
-	FLandscapeCircleHeightPatchPSBase() {}
-	FLandscapeCircleHeightPatchPSBase(const ShaderMetaType::CompiledShaderInitializerType& Initializer)
-		: FGlobalShader(Initializer)
-	{}
-
-public:
 	static bool ShouldCompilePermutation(const FGlobalShaderPermutationParameters& Parameters);
-};
-
-template <bool bPerformBlending = false>
-class FLandscapeCircleHeightPatchPS : public FLandscapeCircleHeightPatchPSBase
-{
-	DECLARE_GLOBAL_SHADER(FLandscapeCircleHeightPatchPS);
-	SHADER_USE_PARAMETER_STRUCT(FLandscapeCircleHeightPatchPS, FLandscapeCircleHeightPatchPSBase);
-
-public:
 	static void ModifyCompilationEnvironment(const FGlobalShaderPermutationParameters& Parameters, FShaderCompilerEnvironment& OutEnvironment);
 
-	static void AddToRenderGraph(FRDGEventName&& InRDGEventName, FRDGBuilder& GraphBuilder, FLandscapeCircleHeightPatchPSBase::FParameters* InParameters, const FIntRect& DestinationBounds);
+	static void AddToRenderGraph(FRDGEventName&& InRDGEventName, FRDGBuilder& GraphBuilder, FParameters* InParameters, const FIntRect& DestinationBounds);
 };
 
 /**
  * Shader that applies a circle patch to a landscape visibility layer.
  */
-template <bool bPerformBlending = false>
-class FLandscapeCircleVisibilityPatchPS : public FLandscapeCircleHeightPatchPSBase
+class FLandscapeCircleVisibilityPatchPS : public FLandscapeCircleHeightPatchPS
 {
 	DECLARE_GLOBAL_SHADER(FLandscapeCircleVisibilityPatchPS);
-	SHADER_USE_PARAMETER_STRUCT(FLandscapeCircleVisibilityPatchPS, FLandscapeCircleHeightPatchPSBase);
 
 public:
+	FLandscapeCircleVisibilityPatchPS() = default;
+	FLandscapeCircleVisibilityPatchPS(const ShaderMetaType::CompiledShaderInitializerType& Initializer)
+		: FLandscapeCircleHeightPatchPS(Initializer)
+	{
+	}
+
 	static void ModifyCompilationEnvironment(const FGlobalShaderPermutationParameters& Parameters, FShaderCompilerEnvironment& OutEnvironment);
-	static void AddToRenderGraph(FRDGEventName&& InRDGEventName, FRDGBuilder& GraphBuilder, FLandscapeCircleHeightPatchPSBase::FParameters* InParameters, const FIntRect& DestinationBounds);
-};
+	static void AddToRenderGraph(FRDGEventName&& InRDGEventName, FRDGBuilder& GraphBuilder, FParameters* InParameters, const FIntRect& DestinationBounds);
+};
\ No newline at end of file
diff --git a/Engine/Plugins/Experimental/LandscapePatch/Source/LandscapePatch/Private/LandscapePatchComponent.cpp b/Engine/Plugins/Experimental/LandscapePatch/Source/LandscapePatch/Private/LandscapePatchComponent.cpp
--- a/Engine/Plugins/Experimental/LandscapePatch/Source/LandscapePatch/Private/LandscapePatchComponent.cpp
+++ b/Engine/Plugins/Experimental/LandscapePatch/Source/LandscapePatch/Private/LandscapePatchComponent.cpp
@@ -40,30 +40,14 @@
  * DetailPanelLayerGuid: the guid representation that the user sees. Should be kept in sync with EditLayerGuid.
  * Priority: double that determines the sorting of the patches.
  * 
- * PatchManager: legacy object that used to hold a list of patches and apply them, where the patch order
- *  was determined by the index in its legacy patch list. If PatchManager is not null, it takes precedence,
- *  so the patch shouldn't be registered to an edit layer or use Priority for sorting. PatchManager should
- *  be mutually exclusive with EditLayerGuid. So if PatchManager is set, it will clear the guid, and vice versa.
+ * PatchManager: deprecated object that used to hold a list of patches and apply them. This is now automatically converted
+ *  to a patch edit layer which handles its patches via a transient list
  * 
- * Landscape: pointer to the landscape in which the guid is found, or in which the PatchManager resides.
+ * Landscape: pointer to the landscape in which the guid is found
  * 
  * Resolving erroneous states:
- * If landscape disagrees with guid or patch manager, guid/patch manager are authoritative, and landscape
- *  should be fixed.
- * If both PatchManager and guid are present, PatchManager wins.
- * If PatchManager is not null but patch is not in the patch manager's list, this is an error, and patch
- *  should be added to manager. However if patch is in a patch manager's list but does not have its PatchManager
- *  pointer set, then the patch should be removed from the PatchManager. Basically, the patch, not the manager,
- *  is authoritative.
  * If the guid does not point to a valid edit layer, it does not get cleared (to be cleaner in cases where an
  *  edit layer might be deleted and then the deletion undone), but the patch will not work until the guid is fixed.
- * 
- * Patch application:
- * In non-batched-merge path, patches receive the legacy Initialize_Native and RenderLayer_Native calls. Legacy
- *  list patches get this from the patch manager, and edit-layer-bound patches get it trhough the edit layer's
- *  legacy adapter interface.
- * In batched-merge path, patches (both in legacy list and in edit layer) provide their desired area of affect
- *  and then get RenderLayer calls
  */
 
 
@@ -77,8 +61,7 @@ namespace LandscapePatchComponentLocals
 #if WITH_EDITOR
 	FAutoConsoleCommand CCmdFixPatchBindings(
 		TEXT("LandscapePatch.FixPatchBindings"),
-		TEXT("For all patches, make sure that patch is either properly bound to a landscape layer or uses a legacy "
-			"patch manager."),
+		TEXT("For all patches, make sure that patch is properly bound to a landscape layer."),
 		FConsoleCommandDelegate::CreateLambda([]() 
 	{
 		const FScopedTransaction Transaction(FixPatchBindingsTransactionName);
@@ -147,68 +130,6 @@ void ULandscapePatchComponent::CheckForErrors()
 		return Arguments;
 	};
 
-	if (PatchManager.IsValid())
-	{
-		ensure(!PatchManager->IsDead());
-		if (EditLayerGuid.IsValid())
-		{
-			FMessageLog("MapCheck").Warning()
-			->AddToken(FUObjectToken::Create(this))
-			->AddToken(FTextToken::Create(FText::Format(LOCTEXT("HaveBothPatchManagerAndGuid", "Patch has both a patch "
-				"manager and a landscape edit layer guid. Patches should use one or the other. Guid is preferred because the "
-				"patch manager will be deprecated. (Package: {Package}, Actor: {Actor}). "
-				"Fix individually or run LandscapePatch.FixPatchBindings."), GetPackageAndActorArgs())))
-			->AddToken(FActionToken::Create(LOCTEXT("SwitchToUsingGuidButton", "Switch to using Guid"), FText(),
-				FOnActionTokenExecuted::CreateWeakLambda(this, [this]()
-			{
-				if (PatchManager.IsValid() && EditLayerGuid.IsValid())
-				{
-					const FScopedTransaction Transaction(RebindPatchTransactionName);
-					ResetPatchManager();
-				}
-			})));
-		}
-
-		if (PatchManager->GetOwningLandscape() != Landscape.Get())
-		{
-			FMessageLog("MapCheck").Warning()
-			->AddToken(FUObjectToken::Create(this))
-			->AddToken(FTextToken::Create(FText::Format(LOCTEXT("PatchManagerAndLandscapeDisagree", "Patch has inconsistent "
-				"manager and landscape pointers. (Package: {Package}, Actor: {Actor}). "
-				"Fix individually or run LandscapePatch.FixPatchBindings."), GetPackageAndActorArgs())))
-			->AddToken(FActionToken::Create(LOCTEXT("FixManagerLandscapeDisagreementButton", "Fix inconsistent pointers"), FText(),
-				FOnActionTokenExecuted::CreateWeakLambda(this, [this]()
-			{
-				if (PatchManager.IsValid() && PatchManager->GetOwningLandscape() != Landscape.Get())
-				{
-					const FScopedTransaction Transaction(RebindPatchTransactionName);
-
-					// It is tempting to just do SetPatchManager(PatchManager.Get()) here, but that will clear
-					// EditLayerGuid as well, which seems like it goes beyond what we promised to fix.
-					Modify();
-					Landscape = PatchManager->GetOwningLandscape();
-				}
-			})));
-		}
-
-		if (!PatchManager->ContainsPatch(this))
-		{
-			FMessageLog("MapCheck").Warning()
-			->AddToken(FUObjectToken::Create(this))
-			->AddToken(FTextToken::Create(FText::Format(LOCTEXT("PatchNotInManager", "Patch has patch manager, but "
-				"was not found in its legacy patch list. (Package: {Package}, Actor: {Actor}). "), GetPackageAndActorArgs())))
-			->AddToken(FActionToken::Create(LOCTEXT("AddToManagerButton", "Add to manager"), FText(),
-				FOnActionTokenExecuted::CreateWeakLambda(this, [this]() 
-			{
-				if (PatchManager.IsValid())
-				{
-					const FScopedTransaction Transaction(RebindPatchTransactionName);
-					PatchManager->AddPatch(this);
-				}
-			})));
-		}
-	}
-
 	if (EditLayerGuid.IsValid())
 	{
 		const ULandscapeEditLayerBase* LocalEditLayer = nullptr;
@@ -242,7 +163,7 @@ void ULandscapePatchComponent::CheckForErrors()
 		}
 	}
 
-	if (!PatchManager.IsValid() && !EditLayerGuid.IsValid())
+	if (!EditLayerGuid.IsValid())
 	{
 		if (Landscape.IsValid())
 		{
@@ -290,27 +211,6 @@ void ULandscapePatchComponent::OnComponentCreated()
 		return;
 	}
 
-	// Legacy: if we're registered through a patch manager (i.e. we're copying a component that is),
-	// make sure our other pointers are consistent.
-	if (PatchManager.IsValid())
-	{
-		check(!PatchManager->IsDead());
-		if (Landscape != PatchManager->GetOwningLandscape())
-		{
-			Modify();
-			Landscape = PatchManager->GetOwningLandscape();
-		}
-
-		// The patch manager might legitimately not contain the set patch manager if we're copying a patch. 
-		if (!PatchManager->ContainsPatch(this) && IsPatchInWorld())
-		{
-			PatchManager->AddPatch(this);
-		}
-
-		ResetEditLayer();
-		return;
-	}
-
 	// We're going to be binding to an edit layer, which will place us in its registered patch list with
 	//  our current priority. If we're later going to be updating our priority to be the highest, then 
 	//  we need to temporarily lower our priority so that we aren't accidentally the highest priority patch
@@ -384,7 +284,7 @@ void ULandscapePatchComponent::PostEditUndo()
 	// Keeps all of our guid-related variables consistent
 	BindToEditLayer(EditLayerGuid);
 
-	// Makes sure we update regardless of whether we're using patch manager or guid
+	// Makes sure we update 
 	RequestLandscapeUpdate();
 }
 
@@ -392,11 +292,7 @@ void ULandscapePatchComponent::OnComponentDestroyed(bool bDestroyingHierarchy)
 {
 	if (!GIsReconstructingBlueprintInstances)
 	{
-		if (PatchManager.IsValid())
-		{
-			PatchManager->RemovePatch(this);
-		}
-		else if (EditLayer.IsValid())
+		if (EditLayer.IsValid())
 		{
 			// Notify the layer that the patch is being destroyed. Note that we are not yet marked
 			//  garbage, so clear our edit layer guid so that the edit layer correctly sees us as
@@ -466,13 +362,6 @@ void ULandscapePatchComponent::PostEditChangeProperty(FPropertyChangedEvent& Pro
 	{
 		SetLandscape(Landscape.Get());
 	}
-	else if (PropertyChangedEvent.Property
-		&& (PropertyChangedEvent.Property->GetFName() == GET_MEMBER_NAME_CHECKED(ULandscapePatchComponent, PatchManager)))
-	{
-		PRAGMA_DISABLE_DEPRECATION_WARNINGS
-		SetPatchManager(PatchManager.Get());
-		PRAGMA_ENABLE_DEPRECATION_WARNINGS
-	}
 	else if (PropertyChangedEvent.Property
 		&& (PropertyChangedEvent.Property->GetFName() == GET_MEMBER_NAME_CHECKED(ULandscapePatchComponent, DetailPanelLayerName)))
 	{
@@ -498,7 +387,7 @@ void ULandscapePatchComponent::PostEditChangeProperty(FPropertyChangedEvent& Pro
 
 	// It is important that this super call happen after the above, because inside a blueprint actor, the call triggers a
 	// rerun of the construction scripts, which will destroy the component and mess with our ability to do the above adjustments
-	// properly (IsValid(this) returns false, the patch manager has the patch removed so it complains when we try to trigger
+	// properly (IsValid(this) returns false, the patch edit layer has the patch removed so it complains when we try to trigger
 	// the update, etc).
 	Super::PostEditChangeProperty(PropertyChangedEvent);
 }
@@ -518,12 +407,6 @@ FLandscapePatchComponentInstanceData::FLandscapePatchComponentInstanceData(const
 	bGaveCouldNotBindToEditLayerWarning = Patch->bGaveCouldNotBindToEditLayerWarning;
 	bGaveMismatchedLandscapeWarning = Patch->bGaveMismatchedLandscapeWarning;
 	bGaveMissingEditLayerGuidWarning = Patch->bGaveMissingEditLayerGuidWarning;
-	bGaveNotInPatchManagerWarning = Patch->bGaveNotInPatchManagerWarning;
-	bGaveMissingLandscapeWarning = Patch->bGaveMissingLandscapeWarning;
-
-	PRAGMA_DISABLE_DEPRECATION_WARNINGS
-	PatchManager = Patch->GetPatchManager();
-	PRAGMA_ENABLE_DEPRECATION_WARNINGS
 
 	EditLayerGuid = Patch->EditLayerGuid;
 	Priority = Patch->Priority;
@@ -551,27 +434,11 @@ void ULandscapePatchComponent::ApplyComponentInstanceData(FLandscapePatchCompone
 	bGaveCouldNotBindToEditLayerWarning = ComponentInstanceData->bGaveCouldNotBindToEditLayerWarning;
 	bGaveMismatchedLandscapeWarning = ComponentInstanceData->bGaveMismatchedLandscapeWarning;
 	bGaveMissingEditLayerGuidWarning = ComponentInstanceData->bGaveMissingEditLayerGuidWarning;
-	bGaveNotInPatchManagerWarning = ComponentInstanceData->bGaveNotInPatchManagerWarning;
-	bGaveMissingLandscapeWarning = ComponentInstanceData->bGaveMissingLandscapeWarning;
 
 	EditLayerGuid = ComponentInstanceData->EditLayerGuid;
 	Priority = ComponentInstanceData->Priority;
 
-	PatchManager = ComponentInstanceData->PatchManager;
-	PreviousPatchManager = PatchManager;
-	if (PatchManager.IsValid())
-	{
-		Landscape = PatchManager->GetOwningLandscape();
-
-		if (IsPatchInWorld() && !PatchManager->ContainsPatch(this))
-		{
-			PatchManager->AddPatch(this);
-		}
-	}
-	else
-	{
-		BindToEditLayer(EditLayerGuid);
-	}
+	BindToEditLayer(EditLayerGuid);
 
 	bInstanceDataApplied = true;
 	if (bDeferUpdateRequestUntilInstanceData)
@@ -594,21 +461,9 @@ void ULandscapePatchComponent::SetLandscape(ALandscape* NewLandscape)
 	if (!NewLandscape)
 	{
 		ResetEditLayer();
-		ResetPatchManager();
 		return;
 	}
 
-	if (PatchManager.IsValid())
-	{
-		if (PatchManager->GetOwningLandscape() == Landscape)
-		{
-			// We're already bound to this landscape through the patch manager
-			return;
-		}
-		// Otherwise, disconnect from the patch manager
-		ResetPatchManager();
-	}
-
 	if (!BindToLandscape(NewLandscape))
 	{
 		UE_LOG(LogLandscapePatch, Warning, TEXT("Unable to bind to given landscape (does it have edit layers enabled?)."));
@@ -616,60 +471,21 @@ void ULandscapePatchComponent::SetLandscape(ALandscape* NewLandscape)
 #endif // WITH_EDITOR
 }
 
-// Changes landscape pointer and layer guid only if given manager is not null
-void ULandscapePatchComponent::SetPatchManager(ALandscapePatchManager* NewPatchManager)
-{
-#if WITH_EDITOR
-	if (!NewPatchManager)
-	{
-		ResetPatchManager();
-		return;
-	}
-
-	Modify(false);
-	
-	// If we're setting a patch manager, we shouldn't be using the normal guid/priority approach.
-	ResetEditLayer();
-
-	if (PreviousPatchManager.IsValid() && PreviousPatchManager != NewPatchManager)
-	{
-		PreviousPatchManager->RemovePatch(this);
-	}
-
-	PatchManager = NewPatchManager;
-	if (IsPatchInWorld())
-	{
-		PatchManager->AddPatch(this);
-	}
-	Landscape = PatchManager->GetOwningLandscape();
-	PreviousPatchManager = NewPatchManager;
-	ResetWarnings();
-#endif // WITH_EDITOR
-}
+// Deprecated
+void ULandscapePatchComponent::SetPatchManager(ADEPRECATED_LandscapePatchManager* NewPatchManager)
+{}
 
 #if WITH_EDITOR
-// Does not touch landscape pointer or layer guid
-void ULandscapePatchComponent::ResetPatchManager()
-{
-	Modify(false);
-
-	if (PreviousPatchManager.IsValid())
-	{
-		PreviousPatchManager->RemovePatch(this);
-	}
-	PatchManager = nullptr;
-	PreviousPatchManager = nullptr;
-}
-
 bool ULandscapePatchComponent::IsPatchPreview()
 {
 	return GetOwner() && GetOwner()->bIsEditorPreviewActor;
 }
 #endif
 
-ALandscapePatchManager* ULandscapePatchComponent::GetPatchManager() const
+// Deprecated
+ADEPRECATED_LandscapePatchManager* ULandscapePatchComponent::GetPatchManager() const
 {
-	return PatchManager.Get();
+	return nullptr;
 }
 
 void ULandscapePatchComponent::RequestLandscapeUpdate(bool bInUserTriggeredUpdate /* = false */)
@@ -697,43 +513,6 @@ void ULandscapePatchComponent::RequestLandscapeUpdate(bool bInUserTriggeredUpdat
 		return;
 	}
 
-	if (PatchManager.IsValid())
-	{
-		// Legacy path of requesting an update through a patch manager
-
-		bool bRequestUpdate = true;
-		if (!PatchManager->ContainsPatch(this))
-		{
-			if (!bGaveNotInPatchManagerWarning)
-			{
-				UE_LOG(LogLandscapePatch, Warning, TEXT("Patch's patch manager does not contain this patch. "
-					"Run LandscapePatch.FixPatchBindings or reset the patch manager on the patch. (Package: %s, Actor: %s)"),
-					*GetNameSafe(GetPackage()), *GetNameSafe(GetAttachmentRootActor()));
-				bGaveNotInPatchManagerWarning = true;
-			}
-			bRequestUpdate = false;
-		}
-		if (!IsValid(PatchManager->GetOwningLandscape()))
-		{
-			if (!bGaveMissingLandscapeWarning)
-			{
-				UE_LOG(LogLandscapePatch, Warning, TEXT("Patch's patch manager does not have a valid owning "
-					"landscape. Run LandscapePatch.FixPatchBindings or reset the landscape on the manager. (Package: %s, Manager: %s)"),
-					*GetNameSafe(GetPackage()), *GetNameSafe(PatchManager.Get()));
-				bGaveMissingLandscapeWarning = true;
-			}
-			bRequestUpdate = false;
-		}
-
-		if (bRequestUpdate)
-		{
-			ResetWarnings();
-			PatchManager->RequestLandscapeUpdate(bInUserTriggeredUpdate || !UE::GetIsEditorLoadingPackage());
-		}
-
-		return;
-	}
-
 	// Otherwise, we work via Guid registration.
 	if (!EditLayer.IsValid() && EditLayerGuid.IsValid())
 	{
@@ -751,8 +530,8 @@ void ULandscapePatchComponent::RequestLandscapeUpdate(bool bInUserTriggeredUpdat
 		if (!bGaveCouldNotBindToEditLayerWarning)
 		{
 			UE_LOG(LogLandscapePatch, Warning, TEXT("Could not find an edit layer with the given guid for the patch. "
-				"(Package: %s, Manager: %s)"),
-				*GetNameSafe(GetPackage()), *GetNameSafe(PatchManager.Get()));
+				"(Package: %s"),
+				*GetNameSafe(GetPackage()));
 			bGaveCouldNotBindToEditLayerWarning = true;
 		}
 	}
@@ -782,12 +561,10 @@ void ULandscapePatchComponent::ResetWarnings()
 	bGaveCouldNotBindToEditLayerWarning = false;
 	bGaveMismatchedLandscapeWarning = false;
 	bGaveMissingEditLayerGuidWarning = false;
-	bGaveNotInPatchManagerWarning = false;
-	bGaveMissingLandscapeWarning = false;
 }
 
-// Safe to do if already bound to the given layer, and safe to do for templates. Does not
-//  touch patch manager. Changes the landscape pointer only if the guid points to a layer, even
+// Safe to do if already bound to the given layer, and safe to do for templates
+//  Changes the landscape pointer only if the guid points to a layer, even
 //  if the layer is of the wrong type. Changes all guid-related variables regardless of whether 
 //  the guid points to a valid layer or not. EditLayerGuid will be valid only if the guid points
 //  to a layer of an appropriate type (in which case binding is considered successful). Binding
@@ -847,7 +624,7 @@ bool ULandscapePatchComponent::BindToEditLayer(const FGuid& Guid)
 	EditLayer = Cast<ULandscapePatchEditLayer>(NewEditLayer);
 	EditLayerGuid = Guid;
 	DetailPanelLayerGuid = Guid.ToString();
-	DetailPanelLayerName =
+	DetailPanelLayerName = 
 		EditLayer.IsValid()
 		? EditLayer->GetName().ToString()
 		: TEXT("-Layer Not Found-");
@@ -899,26 +676,16 @@ bool ULandscapePatchComponent::BindToLandscape(ALandscape* LandscapeIn)
 
 	if (!LandscapeIn)
 	{
-		ResetPatchManager();
 		ResetEditLayer();
 		return true;
 	}
 
 	if (LandscapeIn->IsTemplate())
 	{
-		ResetPatchManager();
 		ResetEditLayer();
 		return false;
 	}
 
-	// See if we're already bound via a patch manager. If not, make sure patch manager is reset.
-	if (PatchManager.IsValid() && PatchManager->GetOwningLandscape() == LandscapeIn)
-	{
-		ResetEditLayer();
-		return true;
-	}
-	ResetPatchManager();
-
 	const ULandscapeEditLayerBase* LocalEditLayer = nullptr;
 	const ULandscapePatchEditLayer* PatchEditLayer = nullptr;
 
@@ -1007,24 +774,6 @@ void ULandscapePatchComponent::FixBindings()
 {
 	Modify(false);
 
-	if (PatchManager.IsValid())
-	{
-		if (!PatchManager->IsDead() && PatchManager->GetOwningLandscape())
-		{
-			ResetEditLayer();
-			Landscape = PatchManager->GetOwningLandscape();
-
-			if (!PatchManager->ContainsPatch(this) && IsPatchInWorld())
-			{
-				PatchManager->AddPatch(this);
-			}
-			return;
-		}
-
-		// Otherwise clear out the invalid manager
-		ResetPatchManager();
-	}
-
 	if (EditLayerGuid.IsValid())
 	{
 		if (BindToEditLayer(EditLayerGuid))
@@ -1037,7 +786,7 @@ void ULandscapePatchComponent::FixBindings()
 		ResetEditLayer();
 	}
 
-	// If we got here, we don't have an edit layer or patch manager
+	// If we got here, we don't have an edit layer 
 	if (Landscape.IsValid())
 	{
 		if (BindToLandscape(Landscape.Get()))
@@ -1064,9 +813,6 @@ void ULandscapePatchComponent::FixBindings()
 
 void ULandscapePatchComponent::NotifyOfBoundLayerDeletion(ULandscapePatchEditLayer* LayerIn)
 {
-	// If we're getting this call, we are presumably bound through an edit layer, not patch manager
-	ensure(!PatchManager.IsValid());
-
 	ResetEditLayer();
 }
 #endif
@@ -1094,11 +840,6 @@ void ULandscapePatchComponent::SetEditLayerGuid(const FGuid& GuidIn)
 #if WITH_EDITOR
 	Modify(false);
 	BindToEditLayer(GuidIn);
-	if (EditLayer.IsValid()) 
-	{
-		// We're not using a legacy manager
-		ResetPatchManager();
-	}
 	ResetWarnings();
 #endif
 }
@@ -1119,22 +860,6 @@ const TArray<FString> ULandscapePatchComponent::GetLayerOptions()
 	TArray<FString> Options{ NullDetailPanelLayerName };
 
 #if WITH_EDITOR
-	// This is a bit of a separate concern, but update the displayed name if necessary. Usually this is not
-	//  needed because it is automatically updated, but it is possible for it to be stuck on a deleted layer 
-	//  in the blueprint editor, where our archetype is not registered to the layer and therefore doesn't get
-	//  the removal notification. 
-	if (!PatchManager.IsValid())
-	{
-		BindToEditLayer(EditLayerGuid);
-	}
-	else
-	{
-		// Uncertain whether we should try to fix a potential conflicting registration issue here (if
-		//  we both have an edit layer and a patch manager), but probably a good idea to, hence the clear
-		//  here. 
-		ResetEditLayer();
-	}
-
 	// Get the layer names from our landscape
 	if (Landscape.IsValid())
 	{
@@ -1153,11 +878,7 @@ const TArray<FString> ULandscapePatchComponent::GetLayerOptions()
 
 FTransform ULandscapePatchComponent::GetLandscapeHeightmapCoordsToWorld() const
 {
-	if (PatchManager.IsValid())
-	{
-		return PatchManager->GetHeightmapCoordsToWorld();
-	}
-	else if (Landscape.IsValid())
+	if (Landscape.IsValid())
 	{
 		UE::Landscape::PatchUtil::GetHeightmapToWorld(Landscape->GetTransform());
 	}
diff --git a/Engine/Plugins/Experimental/LandscapePatch/Source/LandscapePatch/Private/LandscapePatchEditLayer.cpp b/Engine/Plugins/Experimental/LandscapePatch/Source/LandscapePatch/Private/LandscapePatchEditLayer.cpp
--- a/Engine/Plugins/Experimental/LandscapePatch/Source/LandscapePatch/Private/LandscapePatchEditLayer.cpp
+++ b/Engine/Plugins/Experimental/LandscapePatch/Source/LandscapePatch/Private/LandscapePatchEditLayer.cpp
@@ -396,62 +396,4 @@ int32 ULandscapePatchEditLayer::GetInsertionIndex(ULandscapePatchComponent* Patc
 	return Start;
 }
 
-void ULandscapePatchEditLayer::InitializeAsBlueprintBrush(const FTransform& InLandscapeTransform, const FIntPoint& InLandscapeSize, const FIntPoint& InLandscapeRenderTargetSize)
-{
-	HeightmapCoordsToWorld = UE::Landscape::PatchUtil::GetHeightmapToWorld(InLandscapeTransform);
-}
-
-// Called in global merge (legacy) mode to apply the patches
-UTextureRenderTarget2D* ULandscapePatchEditLayer::RenderLayerAsBlueprintBrush(const FLandscapeBrushParameters& InParameters)
-{
-	UpdatePatchListIfDirty();
-
-	FLandscapeBrushParameters BrushParameters = InParameters;
-	for (TSoftObjectPtr<ULandscapePatchComponent>& PatchSoft : RegisteredPatches)
-	{
-		ULandscapePatchComponent* Patch = PatchSoft.Get();
-
-		// We shouldn't have any invalid patches because we should have been notified of dirtying or removals
-		if (!ensure(ShouldPatchBeIncludedInList(Patch)))
-		{
-			bPatchListDirty = true;
-			continue;
-		}
-
-		if (Patch->IsEnabled())
-		{
-			BrushParameters.CombinedResult = Patch->RenderLayer_Native(BrushParameters, HeightmapCoordsToWorld);
-		}
-	}
-
-	return BrushParameters.CombinedResult;
-}
-
-bool ULandscapePatchEditLayer::AffectsHeightmapAsBlueprintBrush() const
-{
-	return Algo::AnyOf(RegisteredPatches, [this](TSoftObjectPtr<ULandscapePatchComponent> PatchSoft)
-	{
-		ULandscapePatchComponent* Patch = PatchSoft.Get();
-		return ShouldPatchBeIncludedInList(Patch) && Patch->AffectsHeightmap();
-	});
-}
-
-bool ULandscapePatchEditLayer::AffectsWeightmapLayerAsBlueprintBrush(const FName& InLayerName) const
-{
-	return Algo::AnyOf(RegisteredPatches, [&InLayerName, this](TSoftObjectPtr<ULandscapePatchComponent> PatchSoft)
-	{ 
-		ULandscapePatchComponent* Patch = PatchSoft.Get();
-		return ShouldPatchBeIncludedInList(Patch) && Patch->AffectsWeightmapLayer(InLayerName);
-	});
-}
-
-bool ULandscapePatchEditLayer::AffectsVisibilityLayerAsBlueprintBrush() const
-{
-	return Algo::AnyOf(RegisteredPatches, [this](TSoftObjectPtr<ULandscapePatchComponent> PatchSoft)
-	{
-		ULandscapePatchComponent* Patch = PatchSoft.Get();
-		return ShouldPatchBeIncludedInList(Patch) && Patch->AffectsVisibilityLayer();
-	});
-}
-
 #endif // WITH_EDITOR
diff --git a/Engine/Plugins/Experimental/LandscapePatch/Source/LandscapePatch/Private/LandscapePatchManager.cpp b/Engine/Plugins/Experimental/LandscapePatch/Source/LandscapePatch/Private/LandscapePatchManager.cpp
--- a/Engine/Plugins/Experimental/LandscapePatch/Source/LandscapePatch/Private/LandscapePatchManager.cpp
+++ b/Engine/Plugins/Experimental/LandscapePatch/Source/LandscapePatch/Private/LandscapePatchManager.cpp
@@ -45,148 +45,6 @@ namespace LandscapePatchManagerLocals
 		true,
 		TEXT("When loaded in the editor, automatically remove all LandscapePatchManagers and bind their patches directly to an edit layer.  Set their patch priorities according to their index."));
 
-#if WITH_EDITOR
-	// Note: the priorities would become jumbled up if someone had multiple managers in the same edit layer. But
-	// this is an unexpected case that is not worth trying to handle differently.
-	FAutoConsoleCommand CCmdMigrateLegacyPatchListToPrioritySystem(
-		TEXT("LandscapePatch.MigrateLegacyListToPrioritySystem"),
-		TEXT("For all patch managers, make any patches in their patch list be directly bound to their edit layer, and "
-			"set the patch priorities according to their index."),
-		FConsoleCommandDelegate::CreateLambda([]()
-	{
-		const FScopedTransaction Transaction(MigratePatchesTransactionName);
-
-		for (TObjectIterator<ALandscapePatchManager> It(
-			/*AdditionalExclusionFlags = */RF_ClassDefaultObject,
-			/*bIncludeDerivedClasses = */true,
-			/*InInternalExclusionFlags = */EInternalObjectFlags::Garbage); It; ++It)
-		{
-			ALandscapePatchManager* Manager = *It;
-			if (!IsValid(Manager))
-			{
-				continue;
-			}
-
-			UWorld* World = Manager->GetWorld();
-			if (Manager->IsTemplate() || !IsValid(World) || World->WorldType != EWorldType::Editor)
-			{
-				continue;
-			}
-
-			if (Manager)
-			{
-				Manager->MigrateToPrioritySystemAndDelete();
-			}
-		}
-	}));
-#endif // WITH_EDITOR
-
-	// Removes invalid patches from the list. This happens automatically when applying patches.
-	void FilterLegacyRegisteredPatches(TArray<TSoftObjectPtr<ULandscapePatchComponent>>& PatchComponents,
-		TMap<TSoftObjectPtr<ULandscapePatchComponent>, int32>& PatchToIndex, const ALandscapePatchManager* ThisPatchManager)
-	{
-		// Used for removing invalid brushes. We remove from the index map immediately but then remove
-		// from the array and update other indices at the very end.
-		bool bHaveInvalidPatches = false;
-		int32 MinRemovedIndex = PatchComponents.Num();
-		auto RemoveComponentFromIndexMap = [&PatchToIndex, &MinRemovedIndex](TSoftObjectPtr<ULandscapePatchComponent>& Component)
-		{
-			int32 RemovedIndex = -1;
-			if (PatchToIndex.RemoveAndCopyValue(Component, RemovedIndex))
-			{
-				MinRemovedIndex = FMath::Min(MinRemovedIndex, RemovedIndex);
-			}
-		};
-
-		for (TSoftObjectPtr<ULandscapePatchComponent>& Component : PatchComponents)
-		{
-			if (Component.IsPending())
-			{
-				Component.LoadSynchronous();
-			}
-
-			if (Component.IsNull())
-			{
-				// Theoretically when components are marked for destruction, they should remove themselves from
-				// the patch manager in their OnComponentDestroyed call. However there seem to be ways to end up
-				// with destroyed patches not being removed, for instance through saving the manager but not the
-				// patch actor.
-				UE_LOG(LogLandscapePatch, Warning, TEXT("ALandscapePatchManager: Found an invalid patch in patch manager. It will be removed."));
-				RemoveComponentFromIndexMap(Component);
-				bHaveInvalidPatches = true;
-				continue;
-			}
-
-			if (!Component.IsValid())
-			{
-				// This means that IsPending() was true, but LoadSynchronous() failed, which we generally don't
-				// expect to happen. However, it can happen in some edge cases such as if you force delete a patch
-				// holder blueprint and don't save the patch manager afterward. Whatever the reason, this is likely
-				// a dead patch that actually needs removal.
-				UE_LOG(LogLandscapePatch, Warning, TEXT("ALandscapePatchManager: Found a pending patch pointer in patch manager that "
-					"turned out to be invalid. It will be removed."));
-				RemoveComponentFromIndexMap(Component);
-				bHaveInvalidPatches = true;
-				continue;
-			}
-
-			// Make sure the patch has this manager set as its patch manager.
-			PRAGMA_DISABLE_DEPRECATION_WARNINGS
-				if (Component->GetPatchManager() != ThisPatchManager)
-					PRAGMA_ENABLE_DEPRECATION_WARNINGS
-				{
-						UE_LOG(LogLandscapePatch, Warning, TEXT("ALandscapePatchManager: Found a patch whose patch manager is not set "
-								"to a patch manager that contains it. It will be removed."));
-						RemoveComponentFromIndexMap(Component);
-						bHaveInvalidPatches = true;
-					continue;
-				}
-
-					if (!Component->IsPatchInWorld())
-					{
-						UE_LOG(LogLandscapePatch, Warning, TEXT("ALandscapePatchManager: Found a non-world patch in patch manager. It will be removed."));
-						RemoveComponentFromIndexMap(Component);
-						bHaveInvalidPatches = true;
-						continue;
-					}
-		}
-
-		if (bHaveInvalidPatches)
-		{
-			PatchComponents.RemoveAll([ThisPatchManager](TSoftObjectPtr<ULandscapePatchComponent> Component) {
-				PRAGMA_DISABLE_DEPRECATION_WARNINGS
-					return !Component.IsValid() || Component->GetPatchManager() != ThisPatchManager || !Component->IsPatchInWorld();
-				PRAGMA_ENABLE_DEPRECATION_WARNINGS
-			});
-			// Update forward indices
-			for (int32 i = MinRemovedIndex; i < PatchComponents.Num(); ++i)
-			{
-				PatchToIndex.Add(PatchComponents[i], i);
-			}
-		}
-	}
-
-	/**
-	 * Returns true if any of the patches in the list pass the predicate.
-	 */
-	bool AnyOfPatchComponents(TArrayView<const TSoftObjectPtr<ULandscapePatchComponent>> LegacyRegisteredPatches,
-		TFunctionRef<bool(ULandscapePatchComponent&)> InPredicate)
-	{
-		return Algo::AnyOf(LegacyRegisteredPatches, [InPredicate](const TSoftObjectPtr<ULandscapePatchComponent>& InComponent)
-		{
-			if (InComponent.IsPending())
-			{
-				InComponent.LoadSynchronous();
-			}
-
-			if (InComponent.IsValid() && InPredicate(*InComponent.Get()))
-			{
-				return true;
-			}
-			return false;
-		});
-	}
-
 #if WITH_EDITOR
 	// Try to get a soft pointer to the owner actor for a ULandscapePatchComponent when that actor and component have not loaded.
 	TSoftObjectPtr<UObject> GetActorPtrFromPatchComponent(TSoftObjectPtr<ULandscapePatchComponent> Patch)
@@ -216,7 +74,7 @@ namespace LandscapePatchManagerLocals
 }
 
 // TODO: Not sure if using this kind of constructor is a proper thing to do vs some other hook...
-ALandscapePatchManager::ALandscapePatchManager(const FObjectInitializer& ObjectInitializer)
+ADEPRECATED_LandscapePatchManager::ADEPRECATED_LandscapePatchManager(const FObjectInitializer& ObjectInitializer)
 	: ALandscapeBlueprintBrushBase(ObjectInitializer)
 {
 #if WITH_EDITOR
@@ -226,77 +84,7 @@ ALandscapePatchManager::ALandscapePatchManager(const FObjectInitializer& ObjectI
 #endif
 }
 
-void ALandscapePatchManager::Initialize_Native(const FTransform& InLandscapeTransform,
-	const FIntPoint& InLandscapeSize,
-	const FIntPoint& InLandscapeRenderTargetSize)
-{
-	HeightmapCoordsToWorld = UE::Landscape::PatchUtil::GetHeightmapToWorld(InLandscapeTransform);
-}
-
-// Called in global merge to apply the patches
-UTextureRenderTarget2D* ALandscapePatchManager::RenderLayer_Native(const FLandscapeBrushParameters& InParameters)
-{
-	using namespace LandscapePatchManagerLocals;
-
-	// Note: We do not expect RenderLayer_Native to be called in the batched merge case.
-	// TODO: Check the cvar to make sure that is not the case?
-
-	FilterLegacyRegisteredPatches(PatchComponents, PatchToIndex, this);
-	FLandscapeBrushParameters BrushParameters = InParameters;
-	for (TSoftObjectPtr<ULandscapePatchComponent>& Component : PatchComponents)
-	{
-		if (!Component->IsEnabled())
-		{
-			// Skip disabled patches
-			continue;
-		}
-
-		BrushParameters.CombinedResult = Component->RenderLayer_Native(BrushParameters, GetHeightmapCoordsToWorld());
-	}
-
-	return BrushParameters.CombinedResult;
-}
-
-#if WITH_EDITOR
-// Called in batched merge path to apply the patches
-TArray<UE::Landscape::EditLayers::FEditLayerRendererState> ALandscapePatchManager::GetEditLayerRendererStates(const UE::Landscape::EditLayers::FMergeContext* InMergeContext)
-{
-	using namespace UE::Landscape::EditLayers;
-	using namespace LandscapePatchManagerLocals;
-
-	FilterLegacyRegisteredPatches(PatchComponents, PatchToIndex, this);
-
-	// Inherit the base brush's enabled state and intersect it with all patches' state, so that the heightmap/weightmap/visibility toggle at the brush's level applies to them too : 
-	FEditLayerTargetTypeState BaseBrushSupportedState(InMergeContext);
-	FEditLayerTargetTypeState BaseBrushEnabledState(InMergeContext);
-	TArray<FTargetLayerGroup> BaseTargetLayerGroups;
-	Super::GetRendererStateInfo(InMergeContext, BaseBrushSupportedState, BaseBrushEnabledState, BaseTargetLayerGroups);
-
-	TArray<FEditLayerRendererState> RendererStates;
-	RendererStates.Reserve(PatchComponents.Num());
-	for (TSoftObjectPtr<ULandscapePatchComponent>& PatchSoft : PatchComponents)
-	{
-		ULandscapePatchComponent* Patch = PatchSoft.Get();
-		if (!Patch)
-		{
-			continue;
-		}
-
-		FEditLayerRendererState& RendererState = RendererStates.Emplace_GetRef(InMergeContext, Patch);
-		// Disable all target types that are disabled in base brush : 
-		RendererState.DisableTargetTypeMask(~BaseBrushEnabledState.GetTargetTypeMask());
-		if (InMergeContext->ShouldSkipProceduralRenderers() || !Patch->IsEnabled())
-		{
-			RendererState.DisableTargetTypeMask(ELandscapeToolTargetTypeFlags::All);
-		}
-	}
-
-	return RendererStates;
-}
-
-#endif // WITH_EDITOR
-
-void ALandscapePatchManager::SetTargetLandscape(ALandscape* InTargetLandscape)
+void ADEPRECATED_LandscapePatchManager::SetTargetLandscape(ALandscape* InTargetLandscape)
 {
 #if WITH_EDITOR
 	if (OwningLandscape != InTargetLandscape && !bDead)
@@ -326,140 +114,14 @@ void ALandscapePatchManager::SetTargetLandscape(ALandscape* InTargetLandscape)
 #endif
 }
 
-bool ALandscapePatchManager::ContainsPatch(ULandscapePatchComponent* Patch) const
-{
-	return PatchToIndex.Contains(TSoftObjectPtr<ULandscapePatchComponent>(Patch));
-}
-
-void ALandscapePatchManager::AddPatch(ULandscapePatchComponent* Patch)
-{
-	using namespace LandscapePatchManagerLocals;
-
-	if (Patch && Patch->IsPatchInWorld() && !IsDead())
-	{
-		if (!ContainsPatch(Patch))
-		{
-			Modify();
-			TSoftObjectPtr<ULandscapePatchComponent> PatchSoftPtr(Patch);
-			PatchComponents.Add(PatchSoftPtr);
-			PatchToIndex.Add(PatchSoftPtr, PatchComponents.Num() - 1);
-		}
-
-		PRAGMA_DISABLE_DEPRECATION_WARNINGS
-			if (Patch->GetPatchManager() != this)
-			{
-				UE_LOG(LogLandscapePatch, Warning, TEXT("ALandscapePatchManager::AddPatch: Added patch does not have this manager set "
-					"as its manager. Patches are typically added to managers by setting the manager on the patch. "
-					"(Package: % s, Actor : % s)"), *GetNameSafe(Patch->GetPackage()), *GetNameSafe(Patch->GetAttachmentRootActor()));
-			}
-		PRAGMA_ENABLE_DEPRECATION_WARNINGS
-
-			// No need to update if the patch is disabled. Important to avoid needlessly updating while dragging a blueprint with
-			// a disabled patch (since construction scripts constantly add and remove).
-			if (Patch->IsEnabled())
-			{
-				RequestLandscapeUpdate(!UE::GetIsEditorLoadingPackage());
-			}
-	}
-	else if (Patch && IsDead())
-	{
-		// IsDead means the migration code ran.  A patch showing up here means it wasn't successfully migrated, possibly from failing to load.
-		UE_LOG(LogLandscapePatch, Error, TEXT("Landscape patch is not correctly connected to a LandscapePatchEditLayer.  Fix this by first clicking Build->Build Landscape and then running LandscapePatch.FixPatchBindings. "
-			"(Package: %s, Actor : %s)"), *GetNameSafe(Patch->GetPackage()), *GetNameSafe(Patch->GetAttachmentRootActor()));
-	}
-}
-
-bool ALandscapePatchManager::RemovePatch(ULandscapePatchComponent* Patch)
-{
-	bool bRemoved = false;
-
-	if (Patch && ContainsPatch(Patch))
-	{
-		Modify();
-		TSoftObjectPtr<ULandscapePatchComponent> PatchSoftPtr(Patch);
-		bRemoved = PatchComponents.Remove(PatchSoftPtr) > 0;
-		if (bRemoved)
-		{
-			int32 RemovedIndex = -1;
-			if (PatchToIndex.RemoveAndCopyValue(PatchSoftPtr, RemovedIndex))
-			{
-				for (int32 i = RemovedIndex; i < PatchComponents.Num(); ++i)
-				{
-					PatchToIndex.Add(PatchComponents[i], i);
-				}
-			}
-
-		}
-
-		// No need to update if the patch was already disabled.Important to avoid needlessly updating while dragging 
-		// a blueprint with a disabled patch (since construction scripts constantly add and remove).
-		if (bRemoved && Patch->IsEnabled())
-		{
-			RequestLandscapeUpdate(!UE::GetIsEditorLoadingPackage());
-		}
-	}
-
-	return bRemoved;
-}
-
-int32 ALandscapePatchManager::GetIndexOfPatch(const ULandscapePatchComponent* Patch) const
-{
-	if (const int32* Index = PatchToIndex.Find(TSoftObjectPtr<ULandscapePatchComponent>(const_cast<ULandscapePatchComponent*>(Patch))))
-	{
-		return *Index;
-	}
-	return INDEX_NONE;
-}
-
-void ALandscapePatchManager::MovePatchToIndex(ULandscapePatchComponent* Patch, int32 Index)
-{
-	if (!Patch || !Patch->IsPatchInWorld() || Index < 0)
-	{
-		return;
-	}
-
-	int32 OriginalIndex = GetIndexOfPatch(Patch);
-	if (OriginalIndex == Index)
-	{
-		return;
-	}
-
-	Modify();
-
-	// It might seem like the index needs adjusting if we're removing before the given index, but that
-	// is not the case if our goal is for the index of the patch to be the given index at the end (rather
-	// than our goal being that the patch be in a particular position relative to the existing patches).
-	RemovePatch(Patch);
-
-	Index = FMath::Clamp(Index, 0, PatchComponents.Num());
-	PatchComponents.Insert(TSoftObjectPtr<ULandscapePatchComponent>(Patch), Index);
-
-	// Update our index lookup structure
-	int32 OtherEndOfChangedIndices = OriginalIndex < 0 ? PatchComponents.Num() - 1
-		: FMath::Min(OriginalIndex, PatchComponents.Num() - 1);
-
-	int32 StartIndex = FMath::Min(Index, OtherEndOfChangedIndices);
-	int32 EndIndex = FMath::Max(Index, OtherEndOfChangedIndices);
-
-	for (int32 i = StartIndex; i <= EndIndex; ++i)
-	{
-		PatchToIndex.Add(PatchComponents[i], i);
-	}
-
-	if (Patch->IsEnabled())
-	{
-		RequestLandscapeUpdate();
-	}
-}
-
 #if WITH_EDITOR
 
-void ALandscapePatchManager::MigrateToPrioritySystemAndDelete()
+void ADEPRECATED_LandscapePatchManager::MigrateToPrioritySystemAndDelete()
 {
 	MigrateToPrioritySystemAndDeleteInternal(/*bAllowUI=*/true);
 }
 
-void ALandscapePatchManager::MigrateToPrioritySystemAndDeleteInternal(bool bAllowUI)
+void ADEPRECATED_LandscapePatchManager::MigrateToPrioritySystemAndDeleteInternal(bool bAllowUI)
 {
 	check(GetWorld()->WorldType == EWorldType::Editor);
 	
@@ -540,10 +202,6 @@ void ALandscapePatchManager::MigrateToPrioritySystemAndDeleteInternal(bool bAllo
 			Patch->SetPriority(Priority);
 			Priority += PriorityStep;
 
-			PRAGMA_DISABLE_DEPRECATION_WARNINGS
-				Patch->SetPatchManager(nullptr);
-			PRAGMA_ENABLE_DEPRECATION_WARNINGS
-
 			Patch->FixBindings();
 		}
 
@@ -559,7 +217,9 @@ void ALandscapePatchManager::MigrateToPrioritySystemAndDeleteInternal(bool bAllo
 	}
 
 	// Important so that we remove ourselves from the landscape blueprint brush list
+PRAGMA_DISABLE_DEPRECATION_WARNINGS
 	SetTargetLandscape(nullptr);
+PRAGMA_ENABLE_DEPRECATION_WARNINGS
 
 	// bDead is used as protection from unexpected weirdness if anything happens in the window between this code running and the
 	// actor being actually removed from the world (by the actionable message update button).
@@ -580,127 +240,7 @@ void ALandscapePatchManager::MigrateToPrioritySystemAndDeleteInternal(bool bAllo
 	}
 }
 
-
-bool ALandscapePatchManager::AffectsHeightmap() const
-{
-	using namespace LandscapePatchManagerLocals;
-
-	if (!Super::AffectsHeightmap())
-	{
-		return false;
-	}
-
-	return AnyOfPatchComponents(PatchComponents,
-		[](ULandscapePatchComponent& InComponent) { return InComponent.AffectsHeightmap(); });
-}
-
-bool ALandscapePatchManager::AffectsWeightmap() const
-{
-	using namespace LandscapePatchManagerLocals;
-
-	if (!Super::AffectsWeightmap())
-	{
-		return false;
-	}
-
-	return AnyOfPatchComponents(PatchComponents,
-		[](ULandscapePatchComponent& InComponent) { return InComponent.AffectsWeightmap(); });
-
-}
-
-bool ALandscapePatchManager::AffectsWeightmapLayer(const FName& InLayerName) const
-{
-	using namespace LandscapePatchManagerLocals;
-
-	if (!Super::AffectsWeightmap()) // Don't call Super::AffectsWeightmapLayer(InLayerName) here as we don't want to use the AffectedWeightmapLayers list for weightmap layers
-	{
-		return false;
-	}
-
-	return AnyOfPatchComponents(PatchComponents,
-		[InLayerName](ULandscapePatchComponent& InComponent) { return InComponent.AffectsWeightmapLayer(InLayerName); });
-}
-
-bool ALandscapePatchManager::AffectsVisibilityLayer() const
-{
-	using namespace LandscapePatchManagerLocals;
-
-	if (!Super::AffectsVisibilityLayer())
-	{
-		return false;
-	}
-
-	return AnyOfPatchComponents(PatchComponents,
-		[](ULandscapePatchComponent& InComponent) { return InComponent.AffectsVisibilityLayer(); });
-}
-
-bool ALandscapePatchManager::CanAffectWeightmapLayer(const FName& InLayerName) const
-{
-	using namespace LandscapePatchManagerLocals;
-
-	if (!Super::CanAffectWeightmap()) // Don't call Super::CanAffectWeightmapLayer(InLayerName) here as we don't want to use the AffectedWeightmapLayers list for weightmap layers
-	{
-		return false;
-	}
-
-	return AnyOfPatchComponents(PatchComponents,
-		[InLayerName](ULandscapePatchComponent& InComponent) { return InComponent.CanAffectWeightmapLayer(InLayerName); });
-}
-
-void ALandscapePatchManager::GetRenderDependencies(TSet<UObject*>& OutDependencies)
-{
-	for (TSoftObjectPtr<ULandscapePatchComponent>& Component : PatchComponents)
-	{
-		if (Component.IsPending())
-		{
-			Component.LoadSynchronous();
-		}
-
-		if (Component.IsValid())
-		{
-			Component->GetRenderDependencies(OutDependencies);
-		}
-	}
-}
-
-void ALandscapePatchManager::SetOwningLandscape(ALandscape* InOwningLandscape)
-{
-	Super::SetOwningLandscape(InOwningLandscape);
-
-	DetailPanelLandscape = OwningLandscape;
-}
-
-// We override PostEditChange to allow the users to change the owning landscape via a property displayed in the detail panel.
-void ALandscapePatchManager::PostEditChangeProperty(FPropertyChangedEvent& PropertyChangedEvent)
-{
-	Super::PostEditChangeProperty(PropertyChangedEvent);
-
-	// Do a bunch of checks to make sure that we don't try to do anything when the editing is happening inside the blueprint editor.
-	UWorld* World = GetWorld();
-	if (IsTemplate() || !IsValidChecked(this) || !IsValid(World) || World->WorldType != EWorldType::Editor)
-	{
-		return;
-	}
-
-	if (PropertyChangedEvent.Property
-		&& (PropertyChangedEvent.Property->GetFName() == GET_MEMBER_NAME_CHECKED(ALandscapePatchManager, DetailPanelLandscape)))
-	{
-		SetTargetLandscape(DetailPanelLandscape.Get());
-	}
-}
-
-void ALandscapePatchManager::PostLoad()
-{
-	Super::PostLoad();
-
-	PatchToIndex.Reset();
-	for (int32 i = 0; i < PatchComponents.Num(); ++i)
-	{
-		PatchToIndex.Add(PatchComponents[i], i);
-	}
-}
-
-void ALandscapePatchManager::PostRegisterAllComponents()
+void ADEPRECATED_LandscapePatchManager::PostRegisterAllComponents()
 {
 	Super::PostRegisterAllComponents();
 
@@ -710,73 +250,6 @@ void ALandscapePatchManager::PostRegisterAllComponents()
 		MigrateToPrioritySystemAndDeleteInternal(/*bAllowUI=*/false);
 	}
 }
-
-void ALandscapePatchManager::CheckForErrors()
-{
-	using namespace LandscapePatchManagerLocals;
-
-	Super::CheckForErrors();
-
-	auto GetPackageAndActorArgs = [this]()
-	{
-		FFormatNamedArguments Arguments;
-		Arguments.Add(TEXT("Package"), FText::FromString(*GetNameSafe(GetPackage())));
-		Arguments.Add(TEXT("Actor"), FText::FromString(*GetNameSafe(this)));
-		return Arguments;
-	};
-
-	// See if we're holding on to any patches that don't have us as the owning patch manager
-	bool bHavePatchWithIncorrectManager = false;
-	for (TSoftObjectPtr<ULandscapePatchComponent>& Component : PatchComponents)
-	{
-PRAGMA_DISABLE_DEPRECATION_WARNINGS
-		if (Component.IsValid() && Component->GetPatchManager() != this)
-		{
-			bHavePatchWithIncorrectManager = true;
-			break;
-		}
-PRAGMA_ENABLE_DEPRECATION_WARNINGS
-	}
-
-	if (bHavePatchWithIncorrectManager)
-	{
-		FMessageLog("MapCheck").Warning()
-			->AddToken(FUObjectToken::Create(this))
-			->AddToken(FTextToken::Create(FText::Format(LOCTEXT("PatchesHaveIncorrectManagerPointer", "Patch manager holds at "
-				"least one patch whose patch manager pointer is set incorrectly. These patches should be removed from the manager."
-				"(Package: {Package}, Manager: {Actor})."), GetPackageAndActorArgs())))
-			->AddToken(FActionToken::Create(LOCTEXT("FixPatchesButton", "Fix patches"), FText(),
-				FOnActionTokenExecuted::CreateWeakLambda(this, [this]()
-				{
-					// Hard to say whether this should be in a transaction, or even be an action, because this happens
-					// automatically on the next landscape update... We'll stick with having it be user triggerable but
-					// not undoable.
-
-					PatchComponents.RemoveAll([this](TSoftObjectPtr<ULandscapePatchComponent> Component) 
-					{
-PRAGMA_DISABLE_DEPRECATION_WARNINGS
-						return Component.IsValid() && Component->GetPatchManager() != this;
-PRAGMA_ENABLE_DEPRECATION_WARNINGS
-					});
-				})));
-	}
-
-	if (PatchComponents.Num() > 0)
-	{
-		FMessageLog("MapCheck").Warning()
-			->AddToken(FUObjectToken::Create(this))
-			->AddToken(FTextToken::Create(FText::Format(LOCTEXT("UsingLegacyPatchList", "The use of the patch manager to "
-				"determine patch ordering is deprecated. Patches should point to a specific edit layer via a guid and "
-				"use Priority for ordering. You can use LandscapePatch.MigrateLegacyListToPrioritySystem to "
-				"fix this. (Package: {Package}, Manager: {Actor})."), GetPackageAndActorArgs())))
-			->AddToken(FActionToken::Create(LOCTEXT("MigrateToGuidsButton", "Migrate to guid system"), FText(),
-				FOnActionTokenExecuted::CreateWeakLambda(this, [this]()
-		{
-			const FScopedTransaction Transaction(MigratePatchesTransactionName);
-			MigrateToPrioritySystemAndDelete();
-		})));
-	}
-}
 #endif // WITH_EDITOR
 
 #undef LOCTEXT_NAMESPACE
diff --git a/Engine/Plugins/Experimental/LandscapePatch/Source/LandscapePatch/Private/LandscapeTexturePatch.cpp b/Engine/Plugins/Experimental/LandscapePatch/Source/LandscapePatch/Private/LandscapeTexturePatch.cpp
--- a/Engine/Plugins/Experimental/LandscapePatch/Source/LandscapePatch/Private/LandscapeTexturePatch.cpp
+++ b/Engine/Plugins/Experimental/LandscapePatch/Source/LandscapePatch/Private/LandscapeTexturePatch.cpp
@@ -193,7 +193,7 @@ bool ULandscapeTexturePatch::RenderLayer(UE::Landscape::EditLayers::FRenderParam
 		{
 			checkf(RDGBuilderRecorder.IsRecording(), TEXT("We should be using ERenderFlags::RenderMode_Recorded in the typical case"));
 			bool bHasRenderedSomething = false;
-			ApplyToHeightmap(/*bInPerformBlending = */false, &RenderParams, RDGBuilderRecorder, CurrentData, GetHeightmapToWorld(RenderParams.RenderAreaWorldTransform), bHasRenderedSomething);
+			ApplyToHeightmap(&RenderParams, RDGBuilderRecorder, CurrentData, GetHeightmapToWorld(RenderParams.RenderAreaWorldTransform), bHasRenderedSomething);
 			return bHasRenderedSomething;
 		}
 	}
@@ -246,7 +246,7 @@ bool ULandscapeTexturePatch::RenderLayer(UE::Landscape::EditLayers::FRenderParam
 		else
 		{
 			checkf(RDGBuilderRecorder.IsRecording(), TEXT("We should be using ERenderFlags::RenderMode_Recorded in the typical case"));
-			ApplyToWeightmap(/*bInPerformBlending = */false, &RenderParams, RDGBuilderRecorder, FoundWeightPatch,
+			ApplyToWeightmap(&RenderParams, RDGBuilderRecorder, FoundWeightPatch,
 				TextureArray->GetResource(),
 				TargetLayerIndexInGroup,
 				RenderParams.RenderAreaSectionRect.Size(), 
@@ -282,92 +282,7 @@ void ULandscapeTexturePatch::BlendLayer(UE::Landscape::EditLayers::FRenderParams
 	RenderParams.MergeRenderContext->GenericBlendLayer(BlendParams, RenderParams, RDGBuilderRecorder);
 }
 
-// Legacy path, which gets the entire heightmap.
-UTextureRenderTarget2D* ULandscapeTexturePatch::RenderLayer_Native(const FLandscapeBrushParameters& InParameters, 
-	const FTransform& LandscapeHeightmapToWorld)
-{
-	using namespace UE::Landscape;
-
-	if (!IsPatchInWorld() || !IsEnabled())
-	{
-		return InParameters.CombinedResult;
-	}
-
-	const bool bIsHeightmapTarget = InParameters.LayerType == ELandscapeToolTargetType::Heightmap;
-	const bool bIsWeightmapTarget = InParameters.LayerType == ELandscapeToolTargetType::Weightmap;
-	const bool bIsVisibilityLayerTarget = InParameters.LayerType == ELandscapeToolTargetType::Visibility;
-
-	if (bIsHeightmapTarget)
-	{
-		if (bReinitializeHeightOnNextRender)
-		{
-			bReinitializeHeightOnNextRender = false;
-			ReinitializeHeight(InParameters.CombinedResult, LandscapeHeightmapToWorld);
-			return InParameters.CombinedResult;
-		}
-		else
-		{
-			// We render in immediate mode for the legacy path so we use an immediate recorder :  
-			FRDGBuilderRecorder RDGBuilderRecorderImmediate;
-			bool bHasRenderedSomething = false;
-			return ApplyToHeightmap(/*bInPerformBlending = */true, /*RenderParams = */nullptr, RDGBuilderRecorderImmediate, InParameters.CombinedResult, LandscapeHeightmapToWorld, bHasRenderedSomething, /*OutputAccess = */ERHIAccess::SRVMask); // in the legacy path the RT leaves as a SRV
-		}
-	}
-	else
-	{
-		// Try to find the weight patch
-		ULandscapeWeightPatchTextureInfo* FoundWeightPatch = nullptr;
-
-		for (const TObjectPtr<ULandscapeWeightPatchTextureInfo>& WeightPatch : WeightPatches)
-		{
-			if (!IsValid(WeightPatch))
-			{
-				continue;
-			}
-
-			if ((bIsWeightmapTarget && (WeightPatch->WeightmapLayerName == InParameters.WeightmapLayerName)) ||
-				(bIsVisibilityLayerTarget && WeightPatch->bEditVisibilityLayer))
-			{
-				FoundWeightPatch = WeightPatch;
-				break;
-			}
-		}
-
-		if (!FoundWeightPatch)
-		{
-			return InParameters.CombinedResult;
-		}
-
-		if (FoundWeightPatch->bReinitializeOnNextRender)
-		{
-			FoundWeightPatch->bReinitializeOnNextRender = false;
-			if (ensure(InParameters.CombinedResult->GetResource()))
-			{
-				ReinitializeWeightPatch(FoundWeightPatch, InParameters.CombinedResult->GetResource(),
-					FIntPoint(InParameters.CombinedResult->SizeX, InParameters.CombinedResult->SizeY), 
-					-1, // Signifies that this is not a Texture2DArray
-					LandscapeHeightmapToWorld);
-			}
-			return InParameters.CombinedResult;
-		}
-		else
-		{
-			// We render in immediate mode for the legacy path so we use an immediate recorder :  
-			FRDGBuilderRecorder RDGBuilderRecorderImmediate;
-			bool bHasRenderedSomething = false;
-			ApplyToWeightmap(/*bInPerformBlending = */true, /*RenderParams = */nullptr, RDGBuilderRecorderImmediate, FoundWeightPatch,
-				InParameters.CombinedResult->GetResource(),
-				0, // Slice index
-				FIntPoint(InParameters.CombinedResult->SizeX, InParameters.CombinedResult->SizeY),
-				LandscapeHeightmapToWorld, 
-				bHasRenderedSomething,
-				/*OutputAccess = */ERHIAccess::SRVMask); // in the legacy path the RT leaves as a SRV
-			return InParameters.CombinedResult;
-		}
-	}
-}
-
-UTextureRenderTarget2D* ULandscapeTexturePatch::ApplyToHeightmap(bool bInPerformBlending, UE::Landscape::EditLayers::FRenderParams* RenderParams, UE::Landscape::FRDGBuilderRecorder& RDGBuilderRecorder, UTextureRenderTarget2D* InCombinedResult,
+UTextureRenderTarget2D* ULandscapeTexturePatch::ApplyToHeightmap(UE::Landscape::EditLayers::FRenderParams* RenderParams, UE::Landscape::FRDGBuilderRecorder& RDGBuilderRecorder, UTextureRenderTarget2D* InCombinedResult,
 	const FTransform& LandscapeHeightmapToWorld, bool& bHasRenderedSomething, ERHIAccess OutputAccess)
 {
 	using namespace UE::Landscape;
@@ -422,8 +337,6 @@ UTextureRenderTarget2D* ULandscapeTexturePatch::ApplyToHeightmap(bool bInPerform
 
 	FTextureResource* OutputResource = InCombinedResult->GetResource();
 	FString OutputResourceName = OutputResource->GetResourceName().ToString();
-	if (!bInPerformBlending)
-	{
 		check(RenderParams != nullptr);
 		ULandscapeScratchRenderTarget* WriteRT = RenderParams->MergeRenderContext->GetBlendRenderTargetWrite();
 		// After this point, the render cannot fail so if we're the first in our render layer group to render, we can cycle the blend render targets and 
@@ -437,43 +350,21 @@ UTextureRenderTarget2D* ULandscapeTexturePatch::ApplyToHeightmap(bool bInPerform
 		}
 		OutputResource = WriteRT->GetRenderTarget()->GetResource();
 		OutputResourceName = WriteRT->GetDebugName();
-	}
 
 	auto RDGCommand =
 		[ OutputResource
 		, OutputResourceName
 		, ShaderParamsToCopy
 		, PatchTextureResource
-		, DestinationBounds
-		, bInPerformBlending](FRDGBuilder& GraphBuilder)
+		, DestinationBounds](FRDGBuilder& GraphBuilder)
 		{
 			TRACE_CPUPROFILER_EVENT_SCOPE(LandscapeTextureHeightPatch_Render);
 
 			TRefCountPtr<IPooledRenderTarget> DestinationRenderTarget = CreateRenderTarget(OutputResource->GetTexture2DRHI(), TEXT("LandscapeTextureHeightPatchOutput"));
 			FRDGTextureRef DestinationTexture = GraphBuilder.RegisterExternalTexture(DestinationRenderTarget);
 
-			FRDGTextureSRVRef InputCopySRV = nullptr;
-			if (bInPerformBlending)
-			{
-				// Make a copy of the portion of our heightmap input that we're writing to so that we can 
-				// read and write at the same time (needed for blending)
-				FRDGTextureDesc InputCopyDescription = DestinationTexture->Desc;
-				InputCopyDescription.NumMips = 1;
-				InputCopyDescription.Extent = DestinationBounds.Size();
-				FRDGTextureRef InputCopy = GraphBuilder.CreateTexture(InputCopyDescription, TEXT("LandscapeTextureHeightPatchInputCopy"));
-				InputCopySRV = GraphBuilder.CreateSRV(FRDGTextureSRVDesc::CreateForMipLevel(InputCopy, 0));
-
-				FRHICopyTextureInfo CopyTextureInfo;
-				CopyTextureInfo.SourceMipIndex = 0;
-				CopyTextureInfo.NumMips = 1;
-				CopyTextureInfo.SourcePosition = FIntVector(DestinationBounds.Min.X, DestinationBounds.Min.Y, 0);
-				CopyTextureInfo.Size = FIntVector(InputCopyDescription.Extent.X, InputCopyDescription.Extent.Y, 0);
-
-				AddCopyTexturePass(GraphBuilder, DestinationTexture, InputCopy, CopyTextureInfo);
-			}
-
-			FApplyLandscapeTextureHeightPatchPSBase::FParameters* ShaderParams =
-				GraphBuilder.AllocParameters<FApplyLandscapeTextureHeightPatchPSBase::FParameters>();
+			FApplyLandscapeTextureHeightPatchPS::FParameters* ShaderParams =
+				GraphBuilder.AllocParameters<FApplyLandscapeTextureHeightPatchPS::FParameters>();
 			*ShaderParams = ShaderParamsToCopy;
 
 			TRefCountPtr<IPooledRenderTarget> PatchRenderTarget = CreateRenderTarget(PatchTextureResource->GetTexture2DRHI(), TEXT("LandscapeTextureHeightPatch"));
@@ -482,24 +373,9 @@ UTextureRenderTarget2D* ULandscapeTexturePatch::ApplyToHeightmap(bool bInPerform
 			ShaderParams->InHeightPatch = PatchSRV;
 			ShaderParams->InHeightPatchSampler = TStaticSamplerState<SF_Bilinear, AM_Clamp, AM_Clamp>::GetRHI();
 
-			ShaderParams->InSourceHeightmap = GraphBuilder.CreateSRV(FRDGTextureSRVDesc::Create(GSystemTextures.GetBlackDummy(GraphBuilder)));
-			ShaderParams->InSourceHeightmapOffset = FIntPoint(ForceInit);
-			if (bInPerformBlending)
-			{
-				ShaderParams->InSourceHeightmap = InputCopySRV;
-				ShaderParams->InSourceHeightmapOffset = DestinationBounds.Min;
-			}
-
 			ShaderParams->RenderTargets[0] = FRenderTargetBinding(DestinationTexture, ERenderTargetLoadAction::ENoAction, /*InMipIndex = */0);
 
-			if (bInPerformBlending)
-			{
-				FApplyLandscapeTextureHeightPatchPS</*bPerformBlending = */true>::AddToRenderGraph(RDG_EVENT_NAME("RenderTextureHeightPatch -> %s", *OutputResourceName), GraphBuilder, ShaderParams, DestinationBounds);
-			}
-			else
-			{
-				FApplyLandscapeTextureHeightPatchPS</*bPerformBlending = */false>::AddToRenderGraph(RDG_EVENT_NAME("RenderTextureHeightPatch -> %s", *OutputResourceName), GraphBuilder, ShaderParams, DestinationBounds);
-			}
+			FApplyLandscapeTextureHeightPatchPS::AddToRenderGraph(RDG_EVENT_NAME("RenderTextureHeightPatch -> %s", *OutputResourceName), GraphBuilder, ShaderParams, DestinationBounds);
 		};
 	// We need to specify the final state of the external texture to prevent the graph builder from transitioning it to SRVMask :
 	RDGBuilderRecorder.EnqueueRDGCommand(RDGCommand, { { OutputResource, (OutputAccess == ERHIAccess::None) ? ERHIAccess::RTV : OutputAccess} });
@@ -508,7 +384,7 @@ UTextureRenderTarget2D* ULandscapeTexturePatch::ApplyToHeightmap(bool bInPerform
 	return InCombinedResult;
 }
 
-void ULandscapeTexturePatch::ApplyToWeightmap(bool bInPerformBlending, UE::Landscape::EditLayers::FRenderParams* RenderParams, UE::Landscape::FRDGBuilderRecorder& RDGBuilderRecorder, ULandscapeWeightPatchTextureInfo* PatchInfo,
+void ULandscapeTexturePatch::ApplyToWeightmap(UE::Landscape::EditLayers::FRenderParams* RenderParams, UE::Landscape::FRDGBuilderRecorder& RDGBuilderRecorder, ULandscapeWeightPatchTextureInfo* PatchInfo,
 	FTextureResource* InMergedLandscapeTextureResource, int32 LandscapeTextureSliceIndex,
 	const FIntPoint& LandscapeTextureResolution, const FTransform& LandscapeHeightmapToWorld, bool& bHasRenderedSomething, ERHIAccess OutputAccess)
 {
@@ -570,8 +446,6 @@ void ULandscapeTexturePatch::ApplyToWeightmap(bool bInPerformBlending, UE::Lands
 
 	FTextureResource* OutputResource = InMergedLandscapeTextureResource;
 	FString OutputResourceName = OutputResource->GetResourceName().ToString();
-	if (!bInPerformBlending)
-	{
 		check(RenderParams != nullptr);
 		ULandscapeScratchRenderTarget* WriteRT = RenderParams->MergeRenderContext->GetBlendRenderTargetWrite();
 		// After this point, the render cannot fail so if we're the first in our render layer group to render, we can cycle the blend render targets and 
@@ -585,7 +459,6 @@ void ULandscapeTexturePatch::ApplyToWeightmap(bool bInPerformBlending, UE::Lands
 		}
 		OutputResource = WriteRT->GetRenderTarget()->GetResource();
 		OutputResourceName = WriteRT->GetDebugName();
-	}
 
 	auto RDGCommand =
 		[ OutputResource
@@ -593,40 +466,15 @@ void ULandscapeTexturePatch::ApplyToWeightmap(bool bInPerformBlending, UE::Lands
 		, LandscapeTextureSliceIndex
 		, ShaderParamsToCopy
 		, PatchTextureResource
-		, DestinationBounds
-		, bInPerformBlending](FRDGBuilder& GraphBuilder)
+		, DestinationBounds](FRDGBuilder& GraphBuilder)
 		{
 			TRACE_CPUPROFILER_EVENT_SCOPE(LandscapeTextureWeightPatch_Render);
 
 			TRefCountPtr<IPooledRenderTarget> DestinationRenderTarget = CreateRenderTarget(OutputResource->GetTextureRHI(), TEXT("LandscapeTextureWeightPatchOutput"));
 			FRDGTextureRef DestinationTexture = GraphBuilder.RegisterExternalTexture(DestinationRenderTarget);
 
-			FRDGTextureSRVRef InputCopySRV = nullptr;
-			if (bInPerformBlending)
-			{
-				// Make a copy of the portion of our weightmap input that we're writing to so that we can 
-				// read and write at the same time (needed for blending)
-				FRDGTextureDesc InputCopyDescription = DestinationTexture->Desc;
-				InputCopyDescription.Dimension = ETextureDimension::Texture2D;
-				InputCopyDescription.ArraySize = 1;
-				InputCopyDescription.NumMips = 1;
-				InputCopyDescription.Extent = DestinationBounds.Size();
-				FRDGTextureRef InputCopy = GraphBuilder.CreateTexture(InputCopyDescription, TEXT("LandscapeTextureWeightPatchInputCopy"));
-
-				FRHICopyTextureInfo CopyTextureInfo;
-				CopyTextureInfo.SourceMipIndex = 0;
-				CopyTextureInfo.NumMips = 1;
-				CopyTextureInfo.SourceSliceIndex = LandscapeTextureSliceIndex;
-				CopyTextureInfo.NumSlices = 1;
-				CopyTextureInfo.SourcePosition = FIntVector(DestinationBounds.Min.X, DestinationBounds.Min.Y, 0);
-				CopyTextureInfo.Size = FIntVector(InputCopyDescription.Extent.X, InputCopyDescription.Extent.Y, 0);
-				AddCopyTexturePass(GraphBuilder, DestinationTexture, InputCopy, CopyTextureInfo);
-
-				InputCopySRV = GraphBuilder.CreateSRV(FRDGTextureSRVDesc::CreateForMipLevel(InputCopy, 0));
-			}
-
-			FApplyLandscapeTextureWeightPatchPSBase::FParameters* ShaderParams =
-				GraphBuilder.AllocParameters<FApplyLandscapeTextureWeightPatchPSBase::FParameters>();
+			FApplyLandscapeTextureWeightPatchPS::FParameters* ShaderParams =
+				GraphBuilder.AllocParameters<FApplyLandscapeTextureWeightPatchPS::FParameters>();
 			*ShaderParams = ShaderParamsToCopy;
 
 			TRefCountPtr<IPooledRenderTarget> PatchRenderTarget = CreateRenderTarget(PatchTextureResource->GetTexture2DRHI(), TEXT("LandscapeTextureWeightPatch"));
@@ -635,25 +483,10 @@ void ULandscapeTexturePatch::ApplyToWeightmap(bool bInPerformBlending, UE::Lands
 			ShaderParams->InWeightPatch = PatchSRV;
 			ShaderParams->InWeightPatchSampler = TStaticSamplerState<SF_Bilinear, AM_Clamp, AM_Clamp>::GetRHI();
 
-			ShaderParams->InSourceWeightmap = GraphBuilder.CreateSRV(FRDGTextureSRVDesc::Create(GSystemTextures.GetBlackDummy(GraphBuilder)));
-			ShaderParams->InSourceWeightmapCoordOffset = FIntPoint(ForceInit);
-			if (bInPerformBlending)
-			{
-				ShaderParams->InSourceWeightmap = InputCopySRV;
-				ShaderParams->InSourceWeightmapCoordOffset = DestinationBounds.Min;
-			}
-
 			ShaderParams->RenderTargets[0] = FRenderTargetBinding(DestinationTexture, ERenderTargetLoadAction::ENoAction,
 				/*InMipIndex = */0, LandscapeTextureSliceIndex);
 
-			if (bInPerformBlending)
-			{
-				FApplyLandscapeTextureWeightPatchPS</*bPerformBlending = */true>::AddToRenderGraph(RDG_EVENT_NAME("RenderTextureWeightPatch -> %s", *OutputResourceName), GraphBuilder, ShaderParams, DestinationBounds);
-			}
-			else
-			{
-				FApplyLandscapeTextureWeightPatchPS</*bPerformBlending = */false>::AddToRenderGraph(RDG_EVENT_NAME("RenderTextureWeightPatch -> %s", *OutputResourceName), GraphBuilder, ShaderParams, DestinationBounds);
-			}
+			FApplyLandscapeTextureWeightPatchPS::AddToRenderGraph(RDG_EVENT_NAME("RenderTextureWeightPatch -> %s", *OutputResourceName), GraphBuilder, ShaderParams, DestinationBounds);
 		};
 
 	// We need to specify the final state of the external texture to prevent the graph builder from transitioning it to SRVMask :
@@ -775,7 +608,7 @@ void ULandscapeTexturePatch::GetHeightShaderParams(const FTransform& LandscapeHe
 	ParamsOut.InBlendMode = static_cast<uint32>(BlendMode);
 
 	// Pack our booleans into a bitfield
-	using EShaderFlags = FApplyLandscapeTextureHeightPatchPSBase::EFlags;
+	using EShaderFlags = FApplyLandscapeTextureHeightPatchPS::EFlags;
 	EShaderFlags Flags = EShaderFlags::None;
 
 	Flags |= (FalloffMode == ELandscapeTexturePatchFalloffMode::RoundedRectangle) ?
@@ -807,7 +640,7 @@ void ULandscapeTexturePatch::GetWeightShaderParams(
 	ParamsOut.InBlendMode = static_cast<uint32>(WeightPatchInfo->bOverrideBlendMode ? WeightPatchInfo->OverrideBlendMode : BlendMode);
 
 	// Pack our booleans into a bitfield
-	using EShaderFlags = FApplyLandscapeTextureHeightPatchPSBase::EFlags;
+	using EShaderFlags = FApplyLandscapeTextureHeightPatchPS::EFlags;
 	EShaderFlags Flags = EShaderFlags::None;
 
 	Flags |= (FalloffMode == ELandscapeTexturePatchFalloffMode::RoundedRectangle) ?
@@ -855,19 +688,19 @@ void ULandscapeTexturePatch::RequestReinitializeHeight()
 #if WITH_EDITOR
 	if (!Super::IsEnabled())
 	{
-		UE_LOG(LogLandscapePatch, Warning, TEXT("ULandscapeTexturePatch::Reinitialize: Cannot reinitialize while disabled."));
+		UE_LOG(LogLandscapePatch, Warning, TEXT("ULandscapeTexturePatch::ReinitializeHeight: Cannot reinitialize while disabled."));
 		return;
 	}
 
 	if (!Landscape.IsValid())
 	{
-		UE_LOG(LogLandscapePatch, Warning, TEXT("ULandscapeTexturePatch::Reinitialize: No associated landscape to initialize from."));
+		UE_LOG(LogLandscapePatch, Warning, TEXT("ULandscapeTexturePatch::ReinitializeHeight: No associated landscape to initialize from."));
 		return;
 	}
 
-	if (!PatchManager.IsValid() && !GetBoundEditLayer())
+	if (!GetBoundEditLayer())
 	{
-		UE_LOG(LogLandscapePatch, Warning, TEXT("ULandscapeTexturePatch::Reinitialize: Not bound to landscape (via edit layer)."));
+		UE_LOG(LogLandscapePatch, Warning, TEXT("ULandscapeTexturePatch::ReinitializeHeight: Not bound to landscape (via edit layer)."));
 		return;
 	}
 
@@ -889,19 +722,19 @@ void ULandscapeTexturePatch::RequestReinitializeWeights()
 #if WITH_EDITOR
 	if (!Super::IsEnabled())
 	{
-		UE_LOG(LogLandscapePatch, Warning, TEXT("ULandscapeTexturePatch::Reinitialize: Cannot reinitialize while disabled."));
+		UE_LOG(LogLandscapePatch, Warning, TEXT("ULandscapeTexturePatch::ReinitializeWeights: Cannot reinitialize while disabled."));
 		return;
 	}
 
 	if (!Landscape.IsValid())
 	{
-		UE_LOG(LogLandscapePatch, Warning, TEXT("ULandscapeTexturePatch::Reinitialize: No associated landscape to initialize from."));
+		UE_LOG(LogLandscapePatch, Warning, TEXT("ULandscapeTexturePatch::ReinitializeWeights: No associated landscape to initialize from."));
 		return;
 	}
 
-	if (!PatchManager.IsValid() && !GetBoundEditLayer())
+	if (!GetBoundEditLayer())
 	{
-		UE_LOG(LogLandscapePatch, Warning, TEXT("ULandscapeTexturePatch::Reinitialize: Not bound to landscape (via edit layer)."));
+		UE_LOG(LogLandscapePatch, Warning, TEXT("ULandscapeTexturePatch::ReinitializeWeights: Not bound to landscape (via edit layer)."));
 		return;
 	}
 
@@ -1984,9 +1817,9 @@ void ULandscapeTexturePatch::AddWeightPatch(const FName& WeightmapLayerName, ELa
 void ULandscapeTexturePatch::RemoveWeightPatch(const FName& InWeightmapLayerName)
 {
 	WeightPatches.RemoveAll([InWeightmapLayerName](const TObjectPtr<ULandscapeWeightPatchTextureInfo>& InWeightPatch)
-		{
+	{ 
 			return IsValid(InWeightPatch) && InWeightPatch->WeightmapLayerName == InWeightmapLayerName;
-		});
+	});
 }
 
 void ULandscapeTexturePatch::RemoveAllWeightPatches()
