Landscape : removed global and local merge code and deprecated anything that was only available for this.

diff --git a/Engine/Plugins/Experimental/LandscapePatch/Source/LandscapePatch/Private/LandscapeTexturePatchPS.h b/Engine/Plugins/Experimental/LandscapePatch/Source/LandscapePatch/Private/LandscapeTexturePatchPS.h
--- a/Engine/Plugins/Experimental/LandscapePatch/Source/LandscapePatch/Private/LandscapeTexturePatchPS.h
+++ b/Engine/Plugins/Experimental/LandscapePatch/Source/LandscapePatch/Private/LandscapeTexturePatchPS.h
@@ -6,7 +6,6 @@
 #include "GlobalShader.h"
 #include "ShaderParameterStruct.h"
 #include "LandscapeTexturePatch.h"
-#include "PixelShaderUtils.h"
 
 
 struct FLandscapeHeightPatchConvertToNativeParams;
@@ -14,329 +13,248 @@ struct FLandscapeHeightPatchConvertToNativeParams;
 namespace UE::Landscape
 {
 
-/**
- * Shader that applies a texture-based height patch to a landscape heightmap.
- */
-class FApplyLandscapeTextureHeightPatchPSBase : public FGlobalShader
-{
-	// TODO: We could consider exposing an additional global alpha setting that we can use to pass in the given
-	// edit layer alpha value... On the other hand, we currently don't bother doing this in any existing blueprint
-	// brushes, and it would be hard to support in a way that doesn't require each blueprint brush to respect it
-	// individually... Not clear whether this is something worth doing yet.
-
-public:
-	// Flags that get packed into a bitfield because we're not allowed to use bool shader parameters:
-	enum class EFlags : uint8
+	/**
+	 * Shader that applies a texture-based height patch to a landscape heightmap.
+	 */
+	class FApplyLandscapeTextureHeightPatchPS : public FGlobalShader
 	{
-		None = 0,
-
-		// When false, falloff is circular.
-		RectangularFalloff = 1 << 0,
+		DECLARE_GLOBAL_SHADER(FApplyLandscapeTextureHeightPatchPS);
+		SHADER_USE_PARAMETER_STRUCT(FApplyLandscapeTextureHeightPatchPS, FGlobalShader);
 
-		// When true, the texture alpha channel is considered for blending (in addition to falloff, if nonzero)
-		ApplyPatchAlpha = 1 << 1,
+	public:
 
-		// When false, the input is directly interpreted as being the height value to process. When true, the height
-		// is unpacked from the red and green channels to make a 16 bit int.
-		InputIsPackedHeight = 1 << 2
+		// Flags that get packed into a bitfield because we're not allowed to use bool shader parameters:
+		enum class EFlags : uint8
+		{
+			None = 0,
+
+			// When false, falloff is circular.
+			RectangularFalloff = 1 << 0,
+
+			// When true, the texture alpha channel is considered for blending (in addition to falloff, if nonzero)
+			ApplyPatchAlpha = 1 << 1,
+
+			// When false, the input is directly interpreted as being the height value to process. When true, the height
+			// is unpacked from the red and green channels to make a 16 bit int.
+			InputIsPackedHeight = 1 << 2
+		};
+
+		// TODO: We could consider exposing an additional global alpha setting that we can use to pass in the given
+		// edit layer alpha value... On the other hand, we currently don't bother doing this in any existing blueprint
+		// brushes, and it would be hard to support in a way that doesn't require each blueprint brush to respect it
+		// individually... Not clear whether this is something worth doing yet.
+
+		BEGIN_SHADER_PARAMETER_STRUCT(FParameters, )
+			SHADER_PARAMETER_RDG_TEXTURE_SRV(Texture2D<float4>, InHeightPatch)
+			SHADER_PARAMETER_SAMPLER(SamplerState, InHeightPatchSampler)
+			SHADER_PARAMETER(FMatrix44f, InHeightmapToPatch)
+			// Value in patch that corresponds to the landscape mid value, which is our "0 height".
+			SHADER_PARAMETER(float, InZeroInEncoding)
+			// Scale to apply to source values relative to the value that represents 0 height.
+			SHADER_PARAMETER(float, InHeightScale)
+			// Offset to apply to height result after applying height scale
+			SHADER_PARAMETER(float, InHeightOffset)
+			// Amount of the patch edge to not apply in UV space. Generally set to 0.5/Dimensions to avoid applying
+			// the edge half-pixels.
+			SHADER_PARAMETER(FVector2f, InEdgeUVDeadBorder)
+			// In world units, the size of the margin across which the alpha falls from 1 to 0
+			SHADER_PARAMETER(float, InFalloffWorldMargin)
+			// Size of the patch in world units (used for falloff)
+			SHADER_PARAMETER(FVector2f, InPatchWorldDimensions)
+			SHADER_PARAMETER(uint32, InBlendMode)
+			// Some combination of the flags (see constants above).
+			SHADER_PARAMETER(uint32, InFlags)
+
+			RENDER_TARGET_BINDING_SLOTS() // Holds our output
+		END_SHADER_PARAMETER_STRUCT()
+
+		static bool ShouldCompilePermutation(const FGlobalShaderPermutationParameters& Parameters);
+		static void ModifyCompilationEnvironment(const FGlobalShaderPermutationParameters& InParameters, FShaderCompilerEnvironment& OutEnvironment);
+
+		static void AddToRenderGraph(FRDGEventName&& InRDGEventName, FRDGBuilder& GraphBuilder, FParameters* InParameters, const FIntRect& DestinationBounds);
 	};
 
-	BEGIN_SHADER_PARAMETER_STRUCT(FParameters, )
-		SHADER_PARAMETER_RDG_TEXTURE_SRV(Texture2D<float4>, InSourceHeightmap)
-		// Offset of the source heightmap relative to the 0,0 location in the destination heightmap, because
-		// the source is likely to be a copied region from some inner part of the destination. This is basically
-		// a SourceHeightmapToDestinationHeightmap coordinate transformation, except that it is always a simple 
-		// integer translation.
-		SHADER_PARAMETER(FIntPoint, InSourceHeightmapOffset)
-		SHADER_PARAMETER_RDG_TEXTURE_SRV(Texture2D<float4>, InHeightPatch)
-		SHADER_PARAMETER_SAMPLER(SamplerState, InHeightPatchSampler)
-		SHADER_PARAMETER(FMatrix44f, InHeightmapToPatch)
-		// Value in patch that corresponds to the landscape mid value, which is our "0 height".
-		SHADER_PARAMETER(float, InZeroInEncoding)
-		// Scale to apply to source values relative to the value that represents 0 height.
-		SHADER_PARAMETER(float, InHeightScale)
-		// Offset to apply to height result after applying height scale
-		SHADER_PARAMETER(float, InHeightOffset)
-		// Amount of the patch edge to not apply in UV space. Generally set to 0.5/Dimensions to avoid applying
-		// the edge half-pixels.
-		SHADER_PARAMETER(FVector2f, InEdgeUVDeadBorder)
-		// In world units, the size of the margin across which the alpha falls from 1 to 0
-		SHADER_PARAMETER(float, InFalloffWorldMargin)
-		// Size of the patch in world units (used for falloff)
-		SHADER_PARAMETER(FVector2f, InPatchWorldDimensions)
-		SHADER_PARAMETER(uint32, InBlendMode)
-		// Some combination of the flags (see constants above).
-		SHADER_PARAMETER(uint32, InFlags)
-
-		RENDER_TARGET_BINDING_SLOTS() // Holds our output
-	END_SHADER_PARAMETER_STRUCT()
-
-	FApplyLandscapeTextureHeightPatchPSBase() {}
-	FApplyLandscapeTextureHeightPatchPSBase(const ShaderMetaType::CompiledShaderInitializerType& Initializer)
-		: FGlobalShader(Initializer)
-	{}
-
-public:
-	static bool ShouldCompilePermutation(const FGlobalShaderPermutationParameters& Parameters);
-	static void ModifyCompilationEnvironment(const FGlobalShaderPermutationParameters& InParameters, FShaderCompilerEnvironment& OutEnvironment);
-
-};
-
-ENUM_CLASS_FLAGS(FApplyLandscapeTextureHeightPatchPSBase::EFlags);
-
-// We need an alias for FApplyLandscapeTextureHeightPatchPS::FParameters so that it can be forward-declared in another header file (cannot forward-declare internal classes)
-class FApplyLandscapeTextureHeightPatchPSParameters : public FApplyLandscapeTextureHeightPatchPSBase::FParameters {};
-
-template <bool bPerformBlending = false>
-class FApplyLandscapeTextureHeightPatchPS : public FApplyLandscapeTextureHeightPatchPSBase
-{
-	using Super = FApplyLandscapeTextureHeightPatchPSBase;
-	DECLARE_GLOBAL_SHADER(FApplyLandscapeTextureHeightPatchPS);
-	SHADER_USE_PARAMETER_STRUCT(FApplyLandscapeTextureHeightPatchPS, FApplyLandscapeTextureHeightPatchPSBase);
+	// We need an alias for FApplyLandscapeTextureHeightPatchPS::FParameters so that it can be forward-declared in another header file (cannot forward-declare internal classes)
+	class FApplyLandscapeTextureHeightPatchPSParameters : public FApplyLandscapeTextureHeightPatchPS::FParameters {};
 
-public:
+	ENUM_CLASS_FLAGS(FApplyLandscapeTextureHeightPatchPS::EFlags);
 
-	static void ModifyCompilationEnvironment(const FGlobalShaderPermutationParameters& InParameters, FShaderCompilerEnvironment& OutEnvironment)
+	/**
+	 * Simple shader that just offsets each height value in a height patch by a constant.
+	 */
+	class FOffsetHeightmapPS : public FGlobalShader
 	{
-		Super::ModifyCompilationEnvironment(InParameters, OutEnvironment);
+		DECLARE_GLOBAL_SHADER(FOffsetHeightmapPS);
+		SHADER_USE_PARAMETER_STRUCT(FOffsetHeightmapPS, FGlobalShader);
 
-		if (bPerformBlending)
-		{
-			OutEnvironment.SetDefine(TEXT("PERFORM_BLENDING"), 1);
-		}
-	}
+	public:
+
+		BEGIN_SHADER_PARAMETER_STRUCT(FParameters, )
+			SHADER_PARAMETER_RDG_TEXTURE_SRV(Texture2D<float4>, InHeightmap)
+			SHADER_PARAMETER(float, InHeightOffset)
+			RENDER_TARGET_BINDING_SLOTS() // Holds our output
+		END_SHADER_PARAMETER_STRUCT()
+
+		static bool ShouldCompilePermutation(const FGlobalShaderPermutationParameters& Parameters);
+		static void ModifyCompilationEnvironment(const FGlobalShaderPermutationParameters& InParameters, FShaderCompilerEnvironment& OutEnvironment);
 
-	static void AddToRenderGraph(FRDGEventName&& InRDGEventName, FRDGBuilder& GraphBuilder, FParameters* InParameters, const FIntRect& DestinationBounds)
+		static void AddToRenderGraph(FRDGBuilder& GraphBuilder, FParameters* InParameters);
+	};
+
+	/**
+	 * Shader that converts a texture stored in some other encoding (where height is in the R channel) to the
+	 * landscape "native" encoding, where height is stored as a 16 bit int split across the R and G channels.
+	 * This is not perfectly reversible (in case of clamping and due to rounding), but it lets us store the
+	 * texture in the way that it would be applied to the landscape (usually).
+	 */
+	class FConvertToNativeLandscapePatchPS : public FGlobalShader
 	{
-		FGlobalShaderMap* ShaderMap = GetGlobalShaderMap(GMaxRHIFeatureLevel);
-		TShaderMapRef<FApplyLandscapeTextureHeightPatchPS<bPerformBlending>> PixelShader(ShaderMap);
-
-		FPixelShaderUtils::AddFullscreenPass(
-			GraphBuilder,
-			ShaderMap,
-			MoveTemp(InRDGEventName),
-			PixelShader,
-			InParameters,
-			DestinationBounds);
-	}
-};
-
-/**
- * Simple shader that just offsets each height value in a height patch by a constant.
- */
-class FOffsetHeightmapPS : public FGlobalShader
-{
-	DECLARE_GLOBAL_SHADER(FOffsetHeightmapPS);
-	SHADER_USE_PARAMETER_STRUCT(FOffsetHeightmapPS, FGlobalShader);
-
-public:
-
-	BEGIN_SHADER_PARAMETER_STRUCT(FParameters, )
-		SHADER_PARAMETER_RDG_TEXTURE_SRV(Texture2D<float4>, InHeightmap)
-		SHADER_PARAMETER(float, InHeightOffset)
-		RENDER_TARGET_BINDING_SLOTS() // Holds our output
-	END_SHADER_PARAMETER_STRUCT()
-
-	static bool ShouldCompilePermutation(const FGlobalShaderPermutationParameters& Parameters);
-	static void ModifyCompilationEnvironment(const FGlobalShaderPermutationParameters& InParameters, FShaderCompilerEnvironment& OutEnvironment);
-
-	static void AddToRenderGraph(FRDGBuilder& GraphBuilder, FParameters* InParameters);
-};
-
-/**
- * Shader that converts a texture stored in some other encoding (where height is in the R channel) to the 
- * landscape "native" encoding, where height is stored as a 16 bit int split across the R and G channels.
- * This is not perfectly reversible (in case of clamping and due to rounding), but it lets us store the
- * texture in the way that it would be applied to the landscape (usually).
- */
-class FConvertToNativeLandscapePatchPS : public FGlobalShader
-{
-	DECLARE_GLOBAL_SHADER(FConvertToNativeLandscapePatchPS);
-	SHADER_USE_PARAMETER_STRUCT(FConvertToNativeLandscapePatchPS, FGlobalShader);
-
-public:
-
-	BEGIN_SHADER_PARAMETER_STRUCT(FParameters, )
-		SHADER_PARAMETER_RDG_TEXTURE_SRV(Texture2D<float4>, InHeightmap)
-		SHADER_PARAMETER(float, InZeroInEncoding)
-		SHADER_PARAMETER(float, InHeightScale)
-		SHADER_PARAMETER(float, InHeightOffset)
-		RENDER_TARGET_BINDING_SLOTS() // Holds our output
-	END_SHADER_PARAMETER_STRUCT()
-
-	static bool ShouldCompilePermutation(const FGlobalShaderPermutationParameters& Parameters);
-	static void ModifyCompilationEnvironment(const FGlobalShaderPermutationParameters& InParameters, FShaderCompilerEnvironment& OutEnvironment);
-
-	static void AddToRenderGraph(FRDGBuilder& GraphBuilder, FRDGTextureRef SourceTexture,
-		FRDGTextureRef DestinationTexture, const FLandscapeHeightPatchConvertToNativeParams& Params);
-};
-
-/**
- * Shader that undoes the conversion done by FConvertToNativeLandscapePatchPS (to the extent possible, since
- * rounding/clamping makes it not perfectly recoverable).
- */
-class FConvertBackFromNativeLandscapePatchPS : public FGlobalShader
-{
-	DECLARE_GLOBAL_SHADER(FConvertBackFromNativeLandscapePatchPS);
-	SHADER_USE_PARAMETER_STRUCT(FConvertBackFromNativeLandscapePatchPS, FGlobalShader);
+		DECLARE_GLOBAL_SHADER(FConvertToNativeLandscapePatchPS);
+		SHADER_USE_PARAMETER_STRUCT(FConvertToNativeLandscapePatchPS, FGlobalShader);
 
-public:
+	public:
 
-	BEGIN_SHADER_PARAMETER_STRUCT(FParameters, )
-		SHADER_PARAMETER_RDG_TEXTURE_SRV(Texture2D<float4>, InHeightmap)
-		SHADER_PARAMETER(float, InZeroInEncoding)
-		SHADER_PARAMETER(float, InHeightScale)
-		SHADER_PARAMETER(float, InHeightOffset)
-		RENDER_TARGET_BINDING_SLOTS() // Holds our output
-	END_SHADER_PARAMETER_STRUCT()
+		BEGIN_SHADER_PARAMETER_STRUCT(FParameters, )
+			SHADER_PARAMETER_RDG_TEXTURE_SRV(Texture2D<float4>, InHeightmap)
+			SHADER_PARAMETER(float, InZeroInEncoding)
+			SHADER_PARAMETER(float, InHeightScale)
+			SHADER_PARAMETER(float, InHeightOffset)
+			RENDER_TARGET_BINDING_SLOTS() // Holds our output
+		END_SHADER_PARAMETER_STRUCT()
 
-	static bool ShouldCompilePermutation(const FGlobalShaderPermutationParameters& Parameters);
-	static void ModifyCompilationEnvironment(const FGlobalShaderPermutationParameters& InParameters, FShaderCompilerEnvironment& OutEnvironment);
+		static bool ShouldCompilePermutation(const FGlobalShaderPermutationParameters& Parameters);
+		static void ModifyCompilationEnvironment(const FGlobalShaderPermutationParameters& InParameters, FShaderCompilerEnvironment& OutEnvironment);
 
-	static void AddToRenderGraph(FRDGBuilder& GraphBuilder, FRDGTextureRef SourceTexture,
-		FRDGTextureRef DestinationTexture, const FLandscapeHeightPatchConvertToNativeParams& Params);
-};
+		static void AddToRenderGraph(FRDGBuilder& GraphBuilder, FRDGTextureRef SourceTexture,
+			FRDGTextureRef DestinationTexture, const FLandscapeHeightPatchConvertToNativeParams& Params);
+	};
 
-class FApplyLandscapeTextureWeightPatchPSBase : public FGlobalShader
-{
-	// TODO: We could consider exposing an additional global alpha setting that we can use to pass in the given
-	// edit layer alpha value... On the other hand, we currently don't bother doing this in any existing blueprint
-	// brushes, and it would be hard to support in a way that doesn't require each blueprint brush to respect it
-	// individually... Not clear whether this is something worth doing yet.
-
-public:
-	// Flags that get packed into a bitfield because we're not allowed to use bool shader parameters:
-	enum class EFlags : uint8
+	/**
+	 * Shader that undoes the conversion done by FConvertToNativeLandscapePatchPS (to the extent possible, since
+	 * rounding/clamping makes it not perfectly recoverable).
+	 */
+	class FConvertBackFromNativeLandscapePatchPS : public FGlobalShader
 	{
-		None = 0,
+		DECLARE_GLOBAL_SHADER(FConvertBackFromNativeLandscapePatchPS);
+		SHADER_USE_PARAMETER_STRUCT(FConvertBackFromNativeLandscapePatchPS, FGlobalShader);
 
-		// When false, falloff is circular.
-		RectangularFalloff = 1 << 0,
+	public:
 
-		// When true, the texture alpha channel is considered for blending (in addition to falloff, if nonzero)
-		ApplyPatchAlpha = 1 << 1,
-	};
+		BEGIN_SHADER_PARAMETER_STRUCT(FParameters, )
+			SHADER_PARAMETER_RDG_TEXTURE_SRV(Texture2D<float4>, InHeightmap)
+			SHADER_PARAMETER(float, InZeroInEncoding)
+			SHADER_PARAMETER(float, InHeightScale)
+			SHADER_PARAMETER(float, InHeightOffset)
+			RENDER_TARGET_BINDING_SLOTS() // Holds our output
+		END_SHADER_PARAMETER_STRUCT()
 
-	BEGIN_SHADER_PARAMETER_STRUCT(FParameters, )
-		SHADER_PARAMETER_RDG_TEXTURE_SRV(Texture2D<float4>, InSourceWeightmap)
-		// Offset of the source weightmap relative to the 0,0 location in the destination weightmap, because
-		// the source is likely to be a copied region from some inner part of the destination. This is basically
-		// a SourceWeightmapToDestinationWeightmap coordinate transformation, except that it is always a simple 
-		// integer translation.
-		SHADER_PARAMETER(FIntPoint, InSourceWeightmapCoordOffset)
-		SHADER_PARAMETER_RDG_TEXTURE_SRV(Texture2D<float4>, InWeightPatch)
-		SHADER_PARAMETER_SAMPLER(SamplerState, InWeightPatchSampler)
-		SHADER_PARAMETER(FMatrix44f, InWeightmapToPatch)
-		// Amount of the patch edge to not apply in UV space. Generally set to 0.5/Dimensions to avoid applying
-		// the edge half-pixels.
-		SHADER_PARAMETER(FVector2f, InEdgeUVDeadBorder)
-		// In world units, the size of the margin across which the alpha falls from 1 to 0
-		SHADER_PARAMETER(float, InFalloffWorldMargin)
-		// Size of the patch in world units (used for falloff)
-		SHADER_PARAMETER(FVector2f, InPatchWorldDimensions)
-		SHADER_PARAMETER(uint32, InBlendMode)
-		// Some combination of the flags (see constants above).
-		SHADER_PARAMETER(uint32, InFlags)
-
-		RENDER_TARGET_BINDING_SLOTS() // Holds our output
-	END_SHADER_PARAMETER_STRUCT()
-
-	FApplyLandscapeTextureWeightPatchPSBase() {}
-	FApplyLandscapeTextureWeightPatchPSBase(const ShaderMetaType::CompiledShaderInitializerType& Initializer)
-		: FGlobalShader(Initializer)
-	{ }
-
-public:
-	static bool ShouldCompilePermutation(const FGlobalShaderPermutationParameters& Parameters);
-	static void ModifyCompilationEnvironment(const FGlobalShaderPermutationParameters& InParameters, FShaderCompilerEnvironment& OutEnvironment);
-};
-
-ENUM_CLASS_FLAGS(FApplyLandscapeTextureWeightPatchPSBase::EFlags);
-
-// We need an alias for FApplyLandscapeTextureWeightPatchPSParameters::FParameters so that it can be forward-declared in another header file (cannot forward-declare internal classes)
-class FApplyLandscapeTextureWeightPatchPSParameters : public FApplyLandscapeTextureWeightPatchPSBase::FParameters {};
-
-template <bool bPerformBlending = false>
-class FApplyLandscapeTextureWeightPatchPS : public FApplyLandscapeTextureWeightPatchPSBase
-{
-	using Super = FApplyLandscapeTextureWeightPatchPSBase;
-	DECLARE_GLOBAL_SHADER(FApplyLandscapeTextureWeightPatchPS);
-	SHADER_USE_PARAMETER_STRUCT(FApplyLandscapeTextureWeightPatchPS, FApplyLandscapeTextureWeightPatchPSBase);
+		static bool ShouldCompilePermutation(const FGlobalShaderPermutationParameters& Parameters);
+		static void ModifyCompilationEnvironment(const FGlobalShaderPermutationParameters& InParameters, FShaderCompilerEnvironment& OutEnvironment);
 
-public:
+		static void AddToRenderGraph(FRDGBuilder& GraphBuilder, FRDGTextureRef SourceTexture,
+			FRDGTextureRef DestinationTexture, const FLandscapeHeightPatchConvertToNativeParams& Params);
+	};
 
-	static void ModifyCompilationEnvironment(const FGlobalShaderPermutationParameters& InParameters, FShaderCompilerEnvironment& OutEnvironment)
+	class FApplyLandscapeTextureWeightPatchPS : public FGlobalShader
 	{
-		Super::ModifyCompilationEnvironment(InParameters, OutEnvironment);
+		DECLARE_GLOBAL_SHADER(FApplyLandscapeTextureWeightPatchPS);
+		SHADER_USE_PARAMETER_STRUCT(FApplyLandscapeTextureWeightPatchPS, FGlobalShader);
 
-		if (bPerformBlending)
+	public:
+
+		// Flags that get packed into a bitfield because we're not allowed to use bool shader parameters:
+		enum class EFlags : uint8
 		{
-			OutEnvironment.SetDefine(TEXT("PERFORM_BLENDING"), 1);
-		}
-	}
+			None = 0,
+
+			// When false, falloff is circular.
+			RectangularFalloff = 1 << 0,
+
+			// When true, the texture alpha channel is considered for blending (in addition to falloff, if nonzero)
+			ApplyPatchAlpha = 1 << 1,
+		};
+
+		// TODO: We could consider exposing an additional global alpha setting that we can use to pass in the given
+		// edit layer alpha value... On the other hand, we currently don't bother doing this in any existing blueprint
+		// brushes, and it would be hard to support in a way that doesn't require each blueprint brush to respect it
+		// individually... Not clear whether this is something worth doing yet.
+
+		BEGIN_SHADER_PARAMETER_STRUCT(FParameters, )
+			SHADER_PARAMETER_RDG_TEXTURE_SRV(Texture2D<float4>, InWeightPatch)
+			SHADER_PARAMETER_SAMPLER(SamplerState, InWeightPatchSampler)
+			SHADER_PARAMETER(FMatrix44f, InWeightmapToPatch)
+			// Amount of the patch edge to not apply in UV space. Generally set to 0.5/Dimensions to avoid applying
+			// the edge half-pixels.
+			SHADER_PARAMETER(FVector2f, InEdgeUVDeadBorder)
+			// In world units, the size of the margin across which the alpha falls from 1 to 0
+			SHADER_PARAMETER(float, InFalloffWorldMargin)
+			// Size of the patch in world units (used for falloff)
+			SHADER_PARAMETER(FVector2f, InPatchWorldDimensions)
+			SHADER_PARAMETER(uint32, InBlendMode)
+			// Some combination of the flags (see constants above).
+			SHADER_PARAMETER(uint32, InFlags)
+
+			RENDER_TARGET_BINDING_SLOTS() // Holds our output
+		END_SHADER_PARAMETER_STRUCT()
+
+		static bool ShouldCompilePermutation(const FGlobalShaderPermutationParameters& Parameters);
+		static void ModifyCompilationEnvironment(const FGlobalShaderPermutationParameters& InParameters, FShaderCompilerEnvironment& OutEnvironment);
+
+		static void AddToRenderGraph(FRDGEventName&& InRDGEventName, FRDGBuilder& GraphBuilder, FParameters* InParameters, const FIntRect& DestinationBounds);
+	};
+
+	// We need an alias for FApplyLandscapeTextureWeightPatchPSParameters::FParameters so that it can be forward-declared in another header file (cannot forward-declare internal classes)
+	class FApplyLandscapeTextureWeightPatchPSParameters : public FApplyLandscapeTextureWeightPatchPS::FParameters {};
 
-	static void AddToRenderGraph(FRDGEventName&& InRDGEventName, FRDGBuilder& GraphBuilder, FParameters* InParameters, const FIntRect& DestinationBounds)
+	ENUM_CLASS_FLAGS(FApplyLandscapeTextureWeightPatchPS::EFlags);
+
+	class FReinitializeLandscapePatchPS : public FGlobalShader
 	{
-		FGlobalShaderMap* ShaderMap = GetGlobalShaderMap(GMaxRHIFeatureLevel);
-		TShaderMapRef<FApplyLandscapeTextureWeightPatchPS<bPerformBlending>> PixelShader(ShaderMap);
-
-		FPixelShaderUtils::AddFullscreenPass(
-			GraphBuilder,
-			ShaderMap,
-			MoveTemp(InRDGEventName),
-			PixelShader,
-			InParameters,
-			DestinationBounds);
-	}
-};
-
-class FReinitializeLandscapePatchPS : public FGlobalShader
-{
-	DECLARE_GLOBAL_SHADER(FReinitializeLandscapePatchPS);
-	SHADER_USE_PARAMETER_STRUCT(FReinitializeLandscapePatchPS, FGlobalShader);
+		DECLARE_GLOBAL_SHADER(FReinitializeLandscapePatchPS);
+		SHADER_USE_PARAMETER_STRUCT(FReinitializeLandscapePatchPS, FGlobalShader);
 
-public:
+	public:
 
-	class FHeightPatch : SHADER_PERMUTATION_BOOL("HEIGHT_PATCH");
-	using FPermutationDomain = TShaderPermutationDomain<FHeightPatch>;
+		class FHeightPatch : SHADER_PERMUTATION_BOOL("HEIGHT_PATCH");
+		using FPermutationDomain = TShaderPermutationDomain<FHeightPatch>;
 
-	BEGIN_SHADER_PARAMETER_STRUCT(FParameters, )
-		SHADER_PARAMETER_RDG_TEXTURE_SRV(Texture2D<float4>, InSource)
-		SHADER_PARAMETER_SAMPLER(SamplerState, InSourceSampler)
-		SHADER_PARAMETER(FMatrix44f, InPatchToSource)
+		BEGIN_SHADER_PARAMETER_STRUCT(FParameters, )
+			SHADER_PARAMETER_RDG_TEXTURE_SRV(Texture2D<float4>, InSource)
+			SHADER_PARAMETER_SAMPLER(SamplerState, InSourceSampler)
+			SHADER_PARAMETER(FMatrix44f, InPatchToSource)
 
-		RENDER_TARGET_BINDING_SLOTS() // Holds our output
-	END_SHADER_PARAMETER_STRUCT()
+			RENDER_TARGET_BINDING_SLOTS() // Holds our output
+		END_SHADER_PARAMETER_STRUCT()
 
-	static bool ShouldCompilePermutation(const FGlobalShaderPermutationParameters& Parameters);
-	static void ModifyCompilationEnvironment(const FGlobalShaderPermutationParameters& InParameters, FShaderCompilerEnvironment& OutEnvironment);
+		static bool ShouldCompilePermutation(const FGlobalShaderPermutationParameters& Parameters);
+		static void ModifyCompilationEnvironment(const FGlobalShaderPermutationParameters& InParameters, FShaderCompilerEnvironment& OutEnvironment);
 
-	static void AddToRenderGraph(FRDGBuilder& GraphBuilder, FParameters* InParameters, bool bHeightPatch);
-};
+		static void AddToRenderGraph(FRDGBuilder& GraphBuilder, FParameters* InParameters, bool bHeightPatch);
+	};
 
-/**
- * Simple shader for copying textures of potentially different resolutions.
- */
-//~ Theoretically CopyToResolveTarget or AddCopyToResolveTargetPass should work, but I was unable to make them
-//~ work without lots of complaints from the renderer.
-class FSimpleTextureCopyPS : public FGlobalShader
-{
-	DECLARE_GLOBAL_SHADER(FSimpleTextureCopyPS);
-	SHADER_USE_PARAMETER_STRUCT(FSimpleTextureCopyPS, FGlobalShader);
+	/**
+	 * Simple shader for copying textures of potentially different resolutions.
+	 */
+	 //~ Theoretically CopyToResolveTarget or AddCopyToResolveTargetPass should work, but I was unable to make them
+	 //~ work without lots of complaints from the renderer.
+	class FSimpleTextureCopyPS : public FGlobalShader
+	{
+		DECLARE_GLOBAL_SHADER(FSimpleTextureCopyPS);
+		SHADER_USE_PARAMETER_STRUCT(FSimpleTextureCopyPS, FGlobalShader);
 
-public:
+	public:
 
-	BEGIN_SHADER_PARAMETER_STRUCT(FParameters, )
-		SHADER_PARAMETER_RDG_TEXTURE_SRV(Texture2D<float4>, InSource)
-		SHADER_PARAMETER_SAMPLER(SamplerState, InSourceSampler)
-		SHADER_PARAMETER(FVector2f, InDestinationResolution)
-		RENDER_TARGET_BINDING_SLOTS() // Holds our output
-	END_SHADER_PARAMETER_STRUCT()
+		BEGIN_SHADER_PARAMETER_STRUCT(FParameters, )
+			SHADER_PARAMETER_RDG_TEXTURE_SRV(Texture2D<float4>, InSource)
+			SHADER_PARAMETER_SAMPLER(SamplerState, InSourceSampler)
+			SHADER_PARAMETER(FVector2f, InDestinationResolution)
+			RENDER_TARGET_BINDING_SLOTS() // Holds our output
+		END_SHADER_PARAMETER_STRUCT()
 
-	static bool ShouldCompilePermutation(const FGlobalShaderPermutationParameters& Parameters);
-	static void ModifyCompilationEnvironment(const FGlobalShaderPermutationParameters& InParameters, FShaderCompilerEnvironment& OutEnvironment);
+		static bool ShouldCompilePermutation(const FGlobalShaderPermutationParameters& Parameters);
+		static void ModifyCompilationEnvironment(const FGlobalShaderPermutationParameters& InParameters, FShaderCompilerEnvironment& OutEnvironment);
 
-	static void AddToRenderGraph(FRDGBuilder& GraphBuilder, FRDGTextureRef SourceTexture, FRDGTextureRef DestinationTexture);
-};
+		static void AddToRenderGraph(FRDGBuilder& GraphBuilder, FRDGTextureRef SourceTexture, FRDGTextureRef DestinationTexture);
+	};
 
-}//end UE::Landscape
+}//end UE::Landscape
\ No newline at end of file
diff --git a/Engine/Plugins/Experimental/LandscapePatch/Source/LandscapePatch/Public/LandscapeCircleHeightPatch.h b/Engine/Plugins/Experimental/LandscapePatch/Source/LandscapePatch/Public/LandscapeCircleHeightPatch.h
--- a/Engine/Plugins/Experimental/LandscapePatch/Source/LandscapePatch/Public/LandscapeCircleHeightPatch.h
+++ b/Engine/Plugins/Experimental/LandscapePatch/Source/LandscapePatch/Public/LandscapeCircleHeightPatch.h
@@ -27,10 +27,7 @@ class ULandscapeCircleHeightPatch : public ULandscapePatchComponent
 
 public:
 
-
 #if WITH_EDITOR
-	virtual UTextureRenderTarget2D* RenderLayer_Native(const FLandscapeBrushParameters& InParameters, const FTransform& HeightmapToWorld) override;
-
 	// ILandscapeEditLayerRenderer, via ULandscapePatchComponent
 	virtual void GetRendererStateInfo(const UE::Landscape::EditLayers::FMergeContext* InMergeContext,
 		UE::Landscape::EditLayers::FEditLayerTargetTypeState& OutSupportedTargetTypeState, 
@@ -70,7 +67,7 @@ protected:
 
 private:
 #if WITH_EDITOR
-	bool ApplyCirclePatch(bool bInPerformBlending, UE::Landscape::EditLayers::FRenderParams* RenderParams, UE::Landscape::FRDGBuilderRecorder& RDGBuilderRecorder, bool bIsVisibilityLayer,
+	bool ApplyCirclePatch(UE::Landscape::EditLayers::FRenderParams* RenderParams, UE::Landscape::FRDGBuilderRecorder& RDGBuilderRecorder, bool bIsVisibilityLayer,
 		FTextureResource* InMergedLandscapeTextureResource, int32 LandscapeTextureSliceIndex,
 		const FIntPoint& DestinationResolution, const FTransform& HeightmapCoordsToWorld, ERHIAccess OutputAccess = ERHIAccess::None);
 #endif // WITH_EDITOR
diff --git a/Engine/Plugins/Experimental/LandscapePatch/Source/LandscapePatch/Public/LandscapePatchComponent.h b/Engine/Plugins/Experimental/LandscapePatch/Source/LandscapePatch/Public/LandscapePatchComponent.h
--- a/Engine/Plugins/Experimental/LandscapePatch/Source/LandscapePatch/Public/LandscapePatchComponent.h
+++ b/Engine/Plugins/Experimental/LandscapePatch/Source/LandscapePatch/Public/LandscapePatchComponent.h
@@ -16,7 +16,7 @@ enum class ECacheApplyPhase;
 enum class ETeleportType : uint8;
 
 class ALandscape;
-class ALandscapePatchManager;
+class ADEPRECATED_LandscapePatchManager;
 class UTextureRenderTarget2D;
 
 UENUM()
@@ -52,10 +52,6 @@ class ULandscapePatchComponent : public USceneComponent
 public:
 	LANDSCAPEPATCH_API ULandscapePatchComponent(const FObjectInitializer& ObjectInitializer = FObjectInitializer::Get());
 
-	// Called from global merge path
-	virtual UTextureRenderTarget2D* RenderLayer_Native(const FLandscapeBrushParameters& InParameters,
-		const FTransform& HeightmapToWorld) { return InParameters.CombinedResult; }
-
 #if WITH_EDITOR
 	// ILandscapeEditLayerRenderer
 	// Subclasses are expected to implement GetEditLayerRendererDebugName, GetRendererStateInfo, GetRenderItems, and RenderLayer
@@ -123,12 +119,12 @@ public:
 	LANDSCAPEPATCH_API virtual void SetLandscape(ALandscape* NewLandscape);
 
 	UE_DEPRECATED(all, "Patch manager use is deprecated. Patches should point to an edit layer using a Guid, and order themselves using priority.")
-	UFUNCTION(BlueprintCallable, Category = "LandscapePatch")
-	LANDSCAPEPATCH_API virtual void SetPatchManager(ALandscapePatchManager* NewPatchManager);
+	UFUNCTION(BlueprintCallable, Category = "LandscapePatch", meta = (DeprecatedFunction, DeprecationMessage = "Patch manager use is deprecated. Patches should point to an edit layer using a Guid, and order themselves using priority."))
+	LANDSCAPEPATCH_API virtual void SetPatchManager(ADEPRECATED_LandscapePatchManager* NewPatchManager);
 
 	UE_DEPRECATED(all, "Patch manager use is deprecated. Patches should point to an edit layer using a Guid, and order themselves using priority.")
-	UFUNCTION(BlueprintCallable, Category = "LandscapePatch")
-	LANDSCAPEPATCH_API virtual ALandscapePatchManager* GetPatchManager() const;
+	UFUNCTION(BlueprintCallable, Category = "LandscapePatch", meta = (DeprecatedFunction, DeprecationMessage = "Patch manager use is deprecated. Patches should point to an edit layer using a Guid, and order themselves using priority."))
+	LANDSCAPEPATCH_API virtual ADEPRECATED_LandscapePatchManager* GetPatchManager() const;
 
 	// Verifies that we're not an archetype or trash or in an invalid world
 	bool IsPatchInWorld() const;
@@ -139,7 +135,7 @@ public:
 #if WITH_EDITOR
 
 	// Cleanup method that can be used to fix the patch's binding to a landscape. An arbitrary landscape
-	// will be bound to if no guid, patch manager, or landscape is set on the patch.
+	// will be bound to if no guid, or landscape is set on the patch.
 	void FixBindings();
 
 	// Allows the patch to update its UI and cached pointer
@@ -171,7 +167,7 @@ protected:
 	ULandscapePatchEditLayer* GetBoundEditLayer() { return EditLayer.Get(); }
 #endif
 
-	// Guid of the edit layer to which the patch is bound, if not using legacy patch manager list.
+	// Guid of the edit layer to which the patch is bound
 	UPROPERTY()
 	FGuid EditLayerGuid;
 
@@ -179,24 +175,13 @@ protected:
 	UPROPERTY(EditAnywhere, Category = Settings, AdvancedDisplay)
 	ELandscapePatchPriorityInitialization PriorityInitialization = ELandscapePatchPriorityInitialization::AcquireHighest;
 
-	// Value that determines the patch ordering relative to other patches. Note that this is
-	// disregarded when using the legacy patch manager list (i.e. Patch Manager should be null).
+	// Value that determines the patch ordering relative to other patches. 
 	UPROPERTY(EditAnywhere, Category = Settings)
 	double Priority = ULandscapePatchEditLayer::PATCH_PRIORITY_BASE;
 
 	UPROPERTY(EditAnywhere, Category = Settings)
 	TSoftObjectPtr<ALandscape> Landscape = nullptr;
 
-	/**
-	 * Deprecated: Patch manager that can hold the patch as part of its legacy patch list.
-	 * 
-	 * Instead, patches should now point to an edit layer via a guid and order themselves using their
-	 * priority value.
-	 */
-	//~ TODO: Deprecate this property, maybe if we have automatic conversion someday?
-	UPROPERTY(EditAnywhere, Category = Settings, AdvancedDisplay)
-	TSoftObjectPtr<ALandscapePatchManager> PatchManager = nullptr;
-
 	/**
 	 * When false, patch does not affect the landscape. Useful for temporarily disabling the patch.
 	 */
@@ -213,9 +198,14 @@ private:
 	UPROPERTY()
 	bool bPropertiesCopiedIndicator = false;
 
-	// Used to properly transition to a different manager when editing it via the detail panel.
+	/**
+	 * Deprecated: Patch manager that can hold the patch as part of its legacy patch list.
+	 * 
+	 * Instead, patches should now point to an edit layer via a guid and order themselves using their
+	 * priority value.
+	 */
 	UPROPERTY()
-	TSoftObjectPtr<ALandscapePatchManager> PreviousPatchManager = nullptr;
+	TSoftObjectPtr<ADEPRECATED_LandscapePatchManager> PatchManager_DEPRECATED = nullptr;
 
 	//~ Name of edit layer displayed to user
 	/**
@@ -243,8 +233,6 @@ private:
 	bool BindToAnyLandscape();
 	void UpdateEditLayerFromDetailPanelLayerName();
 
-	void ResetPatchManager();
-
 	// Tells us whether the patch is attached to an actor that is being dragged around as an editor preview
 	bool IsPatchPreview();
 
@@ -252,8 +240,6 @@ private:
 	bool bGaveCouldNotBindToEditLayerWarning = false;
 	bool bGaveMissingEditLayerGuidWarning = false;
 	bool bGaveMismatchedLandscapeWarning = false;
-	bool bGaveNotInPatchManagerWarning = false;
-	bool bGaveMissingLandscapeWarning = false;
 	void ResetWarnings();
 	bool bInstanceDataApplied = false;
 	bool bDeferUpdateRequestUntilInstanceData = false;
@@ -272,6 +258,14 @@ struct FLandscapePatchComponentInstanceData : public FSceneComponentInstanceData
 
 	FLandscapePatchComponentInstanceData() = default;
 	FLandscapePatchComponentInstanceData(const ULandscapePatchComponent* SourceComponent);
+	
+PRAGMA_DISABLE_DEPRECATION_WARNINGS
+	// These are necessary to prevent deprecation warnings from UE_DEPRECATED(5.7, ...) on some compilers because of auto-generated code. Can be removed after the deprecated members are removed :
+	FLandscapePatchComponentInstanceData(const FLandscapePatchComponentInstanceData&) = default;
+	FLandscapePatchComponentInstanceData(FLandscapePatchComponentInstanceData&&) = default;
+	FLandscapePatchComponentInstanceData& operator=(const FLandscapePatchComponentInstanceData&) = default;
+	FLandscapePatchComponentInstanceData& operator=(FLandscapePatchComponentInstanceData&&) = default;
+PRAGMA_ENABLE_DEPRECATION_WARNINGS
 
 	virtual ~FLandscapePatchComponentInstanceData() = default;
 
@@ -293,7 +287,7 @@ struct FLandscapePatchComponentInstanceData : public FSceneComponentInstanceData
 	// potentially be used in some multiuser code paths.
 
 	UPROPERTY()
-	TSoftObjectPtr<ALandscapePatchManager> PatchManager = nullptr;
+	TSoftObjectPtr<ADEPRECATED_LandscapePatchManager> PatchManager_DEPRECATED = nullptr;
 	UPROPERTY()
 	FGuid EditLayerGuid;
 	// Priority needs to be carried over becuse our tweaks to it in OnComponentCreated cause it to not be
@@ -309,8 +303,11 @@ struct FLandscapePatchComponentInstanceData : public FSceneComponentInstanceData
 	bool bGaveMismatchedLandscapeWarning = false;
 	UPROPERTY()
 	bool bGaveMissingEditLayerGuidWarning = false;
+
+	UE_DEPRECATED(5.7, "Not needed since the deprecation of PatchManager")
 	UPROPERTY()
 	bool bGaveNotInPatchManagerWarning = false;
+	UE_DEPRECATED(5.7, "Not needed since the deprecation of PatchManager")
 	UPROPERTY()
 	bool bGaveMissingLandscapeWarning = false;
 #endif
diff --git a/Engine/Plugins/Experimental/LandscapePatch/Source/LandscapePatch/Public/LandscapePatchEditLayer.h b/Engine/Plugins/Experimental/LandscapePatch/Source/LandscapePatch/Public/LandscapePatchEditLayer.h
--- a/Engine/Plugins/Experimental/LandscapePatch/Source/LandscapePatch/Public/LandscapePatchEditLayer.h
+++ b/Engine/Plugins/Experimental/LandscapePatch/Source/LandscapePatch/Public/LandscapePatchEditLayer.h
@@ -2,7 +2,6 @@
 
 #pragma once
 
-#include "GlobalMergeLegacySupportUtil.h" // ILandscapeBrushRenderCallAdapter_GlobalMergeLegacySupport
 #include "LandscapeEditLayer.h"
 #include "Misc/Guid.h"
 
@@ -16,7 +15,6 @@ class ULandscapePatchComponent;
  */
 UCLASS(MinimalAPI, meta = (ShortToolTip = "Special edit layer for landscape patches"))
 class ULandscapePatchEditLayer : public ULandscapeEditLayerProcedural
-	, public ILandscapeBrushRenderCallAdapter_GlobalMergeLegacySupport
 {
 	GENERATED_BODY()
 
@@ -66,17 +64,6 @@ public:
 
 	// UObject
 	virtual void PostEditUndo() override;
-
-	// ILandscapeBrushRenderCallAdapter_GlobalMergeLegacySupport
-	virtual bool AffectsHeightmapAsBlueprintBrush() const override;
-	virtual bool AffectsWeightmapLayerAsBlueprintBrush(const FName& InLayerName) const override;
-	virtual bool AffectsVisibilityLayerAsBlueprintBrush() const override;
-protected:
-	// ILandscapeBrushRenderCallAdapter_GlobalMergeLegacySupport
-	virtual void InitializeAsBlueprintBrush(const FTransform& InLandscapeTransform,
-		const FIntPoint& InLandscapeSize,
-		const FIntPoint& InLandscapeRenderTargetSize) override;
-	virtual UTextureRenderTarget2D* RenderLayerAsBlueprintBrush(const FLandscapeBrushParameters& InParameters) override;
 #endif // WITH_EDITOR
 
 #if WITH_EDITORONLY_DATA
@@ -109,8 +96,5 @@ private:
 	// A tracker of the highest priority we've seen. When the patch list is not dirty, it will be the priority
 	//  of the last element, but this variable allows us to maintain the value past dirtying.
 	double HighestKnownPriority = PATCH_PRIORITY_BASE;
-
-	// Used in legacy global merge path to pass transform information from Initialize to RenderLayer
-	FTransform HeightmapCoordsToWorld;
 #endif // WITH_EDITOR
 };
diff --git a/Engine/Plugins/Experimental/LandscapePatch/Source/LandscapePatch/Public/LandscapePatchManager.h b/Engine/Plugins/Experimental/LandscapePatch/Source/LandscapePatch/Public/LandscapePatchManager.h
--- a/Engine/Plugins/Experimental/LandscapePatch/Source/LandscapePatch/Public/LandscapePatchManager.h
+++ b/Engine/Plugins/Experimental/LandscapePatch/Source/LandscapePatch/Public/LandscapePatchManager.h
@@ -18,35 +18,41 @@ class ULandscapePatchComponent;
  * point to a special landscape patch edit layer via a guid, and determine their ordering
  * relative to each other using a priority value.
  */
-UCLASS()
-class ALandscapePatchManager : public ALandscapeBlueprintBrushBase
+UCLASS(Deprecated, meta = (DeprecationMessage = "Landscape patches are now fully handled by ULandscapePatchEditLayer and migration is automatic so effectively, this actor is only there for backwards-compatibility/serialization purposes"))
+class ADEPRECATED_LandscapePatchManager : public ALandscapeBlueprintBrushBase
 {
 	GENERATED_BODY()
 
 public:
-	ALandscapePatchManager(const FObjectInitializer& ObjectInitializer = FObjectInitializer::Get());
-
-	// These get called in the global merge path.
-	virtual void Initialize_Native(const FTransform& InLandscapeTransform,
-		const FIntPoint& InLandscapeSize,
-		const FIntPoint& InLandscapeRenderTargetSize) override;
-	virtual UTextureRenderTarget2D* RenderLayer_Native(const FLandscapeBrushParameters& InParameters) override;
+	ADEPRECATED_LandscapePatchManager(const FObjectInitializer& ObjectInitializer = FObjectInitializer::Get());
 
 #if WITH_EDITOR
 	// IEditLayerRendererProvider
-	// Called by the batched-merge application path.
-	virtual TArray<UE::Landscape::EditLayers::FEditLayerRendererState> GetEditLayerRendererStates(const UE::Landscape::EditLayers::FMergeContext* InMergeContext) override;
+	// Deprecated
+	virtual TArray<UE::Landscape::EditLayers::FEditLayerRendererState> GetEditLayerRendererStates(const UE::Landscape::EditLayers::FMergeContext* InMergeContext) override
+	{ 
+		return {}; 
+	}
 
 	// ILandscapeEditLayerRenderer
-	//~ In batched merge, the manager relies on being a renderer provider. It does not need to have its 
-	//~  RenderLayer_Native method called, so we override the implementations of ILandscapeEditLayerRenderer
-	//~  inherited from ALandscapeBlueprintBrushBase to do nothing.
+	// Deprecated
 	virtual void GetRendererStateInfo(const UE::Landscape::EditLayers::FMergeContext* InMergeContext,
 		UE::Landscape::EditLayers::FEditLayerTargetTypeState& OutSupportedTargetTypeState,
 		UE::Landscape::EditLayers::FEditLayerTargetTypeState& OutEnabledTargetTypeState,
-		TArray<UE::Landscape::EditLayers::FTargetLayerGroup>& OutTargetLayerGroups) const override {}
-	virtual TArray<UE::Landscape::EditLayers::FEditLayerRenderItem> GetRenderItems(const UE::Landscape::EditLayers::FMergeContext* InMergeContext) const override { return {}; }
-	virtual bool RenderLayer(UE::Landscape::EditLayers::FRenderParams& RenderParams, UE::Landscape::FRDGBuilderRecorder& RDGBuilderRecorder) override { return false; };
+		TArray<UE::Landscape::EditLayers::FTargetLayerGroup>& OutTargetLayerGroups) const override 
+	{}
+
+	// Deprecated
+	virtual TArray<UE::Landscape::EditLayers::FEditLayerRenderItem> GetRenderItems(const UE::Landscape::EditLayers::FMergeContext* InMergeContext) const override
+	{ 
+		return {}; 
+	}
+
+	// Deprecated
+	virtual bool RenderLayer(UE::Landscape::EditLayers::FRenderParams& RenderParams, UE::Landscape::FRDGBuilderRecorder& RDGBuilderRecorder) override 
+	{ 
+		return false; 
+	};
 
 	// In 5.6 users should not be able to add new Patch Managers anywhere in the editor
 	virtual bool SupportsBlueprintBrushTool() const { return false; }
@@ -54,49 +60,48 @@ public:
 
 	// Adds the brush to the given landscape, removing it from any previous one. This differs from SetOwningLandscape
 	// in that SetOwningLandscape is called by the landscape itself from AddBrushToLayer to update the manager.
-	UFUNCTION(BlueprintCallable, Category = LandscapeManager)
+	UE_DEPRECATED(5.7, "Patch manager is deprecated and replaced by Landscape Patch Edit Layer.")
+	UFUNCTION(BlueprintCallable, Category = LandscapeManager, meta = (DeprecatedFunction, DeprecationMessage = "Patch manager is deprecated and replaced by Landscape Patch Edit Layer."))
 	virtual void SetTargetLandscape(ALandscape* InOwningLandscape);
 
-	// For use by the owned patch objects.
-	/**
-	 * Gets the transform from a point in the heightmap (where x and y are pixel coordinates,
-	 * aka coordinates of the associated vertex, and z is the height as stored in the height
-	 * map, currently a 16 bit integer) to world point based on the current landscape transform.
-	 */
-	virtual FTransform GetHeightmapCoordsToWorld() { return HeightmapCoordsToWorld; }
-
-	UFUNCTION(BlueprintCallable, Category = LandscapePatch)
-	bool ContainsPatch(ULandscapePatchComponent* Patch) const;
-
-	UFUNCTION(BlueprintCallable, Category = LandscapePatch)
-	void AddPatch(ULandscapePatchComponent* Patch);
-
-	UFUNCTION(BlueprintCallable, Category = LandscapePatch)
-	bool RemovePatch(ULandscapePatchComponent* Patch);
+	// Deprecated
+	UE_DEPRECATED(5.7, "LandscapePatchManager is deprecated and replaced by Landscape Patch Edit Layer.")
+	virtual FTransform GetHeightmapCoordsToWorld() 
+	{ 
+		return FTransform(); 
+	}
 
-	/** 
-	 * Gets the index of a particular patch in the manager's stack of patches (later indices get applied after
-	 * earlier ones.
-	 */
-	UFUNCTION(BlueprintCallable, Category = LandscapePatch)
-	int32 GetIndexOfPatch(const ULandscapePatchComponent* Patch) const;
+	UE_DEPRECATED(5.7, "Patch manager is deprecated and replaced by Landscape Patch Edit Layer.")
+	UFUNCTION(BlueprintCallable, Category = LandscapePatch, meta = (DeprecatedFunction, DeprecationMessage = "Patch manager is deprecated and replaced by Landscape Patch Edit Layer."))
+	bool ContainsPatch(ULandscapePatchComponent* Patch) const
+	{
+		return false;
+	}
 
-	/**
-	 * Moves patch to given index in the list of patches held by the manager (so that it is applied at 
-	 * a particular time relative to the others).
-	 */
-	UFUNCTION(BlueprintCallable, Category = LandscapePatch)
-	void MovePatchToIndex(ULandscapePatchComponent* Patch, int32 Index);
+	UE_DEPRECATED(5.7, "Patch manager is deprecated and replaces by Landscape Patch Edit Layer.")
+	UFUNCTION(BlueprintCallable, Category = LandscapePatch, meta = (DeprecatedFunction, DeprecationMessage = "Patch manager is deprecated and replaced by Landscape Patch Edit Layer."))
+	void AddPatch(ULandscapePatchComponent* Patch)
+	{}
 
-	// This patch manager has been migrated out of and should no longer be accessible.
-	bool IsDead() const {
-#if WITH_EDITOR
-		return bDead;
-#else
+	UE_DEPRECATED(5.7, "Patch manager is deprecated and replaces by Landscape Patch Edit Layer.")
+	UFUNCTION(BlueprintCallable, Category = LandscapePatch, meta = (DeprecatedFunction, DeprecationMessage = "Patch manager is deprecated and replaced by Landscape Patch Edit Layer."))
+	bool RemovePatch(ULandscapePatchComponent* Patch)
+	{
 		return false;
-#endif
 	}
 
+	UE_DEPRECATED(5.7, "Patch manager is deprecated and replaces by Landscape Patch Edit Layer.")
+	UFUNCTION(BlueprintCallable, Category = LandscapePatch, meta = (DeprecatedFunction, DeprecationMessage = "Patch manager is deprecated and replaced by Landscape Patch Edit Layer."))
+	int32 GetIndexOfPatch(const ULandscapePatchComponent* Patch) const
+	{
+		return INDEX_NONE;
+	}
+
+	UE_DEPRECATED(5.7, "Patch manager is deprecated and replaces by Landscape Patch Edit Layer.")
+	UFUNCTION(BlueprintCallable, Category = LandscapePatch, meta = (DeprecatedFunction, DeprecationMessage = "Patch manager use is deprecated and replaces by Landscape Patch Edit Layer."))
+	void MovePatchToIndex(ULandscapePatchComponent* Patch, int32 Index)
+	{}
+
 #if WITH_EDITOR
 
 	/**
@@ -108,29 +113,57 @@ public:
 	void MigrateToPrioritySystemAndDelete();
 
 	// ALandscapeBlueprintBrushBase
-	virtual bool CanAffectWeightmapLayer(const FName& InLayerName) const override;
-	virtual bool AffectsHeightmap() const override;
-	virtual bool AffectsWeightmap() const override;
-	virtual bool AffectsWeightmapLayer(const FName& InLayerName) const override;
-	virtual bool AffectsVisibilityLayer() const override;
-	virtual void GetRenderDependencies(TSet<UObject*>& OutDependencies) override;
-	virtual void SetOwningLandscape(class ALandscape* InOwningLandscape) override;
+	virtual bool CanAffectWeightmapLayer(const FName& InLayerName) const override
+	{
+		return false;
+	}
+
+	virtual bool AffectsHeightmap() const override
+	{
+		return false;
+	}
+
+	virtual bool AffectsWeightmap() const override
+	{
+		return false;
+	}
+
+	virtual bool AffectsWeightmapLayer(const FName& InLayerName) const override
+	{
+		return false;
+	}
+
+	virtual bool AffectsVisibilityLayer() const override
+	{
+		return false;
+	}
+
+	virtual void GetRenderDependencies(TSet<UObject*>& OutDependencies) override
+	{}
 
 	// AActor
-	virtual void CheckForErrors() override;
 	virtual void PostRegisterAllComponents() override;
 
-	// UObject
-	virtual void PostEditChangeProperty(FPropertyChangedEvent& PropertyChangedEvent) override;
-	virtual void PostLoad() override;
 #endif
-	virtual bool IsEditorOnly() const override { return true; }
-	virtual bool NeedsLoadForClient() const override { return false; }
-	virtual bool NeedsLoadForServer() const override { return false; }
+	virtual bool IsEditorOnly() const override 
+	{ 
+		return true; 
+	}
+
+	virtual bool NeedsLoadForClient() const override 
+	{ 
+		return false; 
+	}
+
+	virtual bool NeedsLoadForServer() const override 
+	{ 
+		return false; 
+	}
 
 	// This is intentionally lower than PATCH_PRIORITY_BASE so that patches converted from a
 	// patch manager list are applied before other edit layer patches.
 	inline static const double LEGACY_PATCH_PRIORITY_BASE = ULandscapePatchEditLayer::PATCH_PRIORITY_BASE - 10;
+
 private:
 
 	void MigrateToPrioritySystemAndDeleteInternal(bool bAllowUI);
@@ -138,24 +171,8 @@ private:
 	UPROPERTY()
 	TArray<TSoftObjectPtr<ULandscapePatchComponent>> PatchComponents;
 
-	// Used in legacy paths to pass the transform information from Initialize_Native to RenderLayer_Native
-	UPROPERTY()
-	FTransform HeightmapCoordsToWorld;
-
 #if WITH_EDITORONLY_DATA
-	//~ This is transient because SetOwningLandscape is called in ALandscape::PostLoad.
-	/**
-	 * The owning landscape.
-	 */
-	UPROPERTY(EditAnywhere, Category = Landscape, Transient, meta = (DisplayName = "Landscape"))
-	TObjectPtr<ALandscape> DetailPanelLandscape = nullptr;
-
-	bool bIssuedPatchOwnershipWarning = false;
-
+	// This patch manager has been migrated out of and should no longer be accessible.
 	bool bDead = false;
 #endif
-
-	// Transient table to speed up Contains and IndexOf queries, which are very slow for an array of TSoftObjectPtr's.
-	UPROPERTY(Transient)
-	TMap<TSoftObjectPtr<ULandscapePatchComponent>, int32> PatchToIndex;
 };
diff --git a/Engine/Plugins/Experimental/LandscapePatch/Source/LandscapePatch/Public/LandscapeTexturePatch.h b/Engine/Plugins/Experimental/LandscapePatch/Source/LandscapePatch/Public/LandscapeTexturePatch.h
--- a/Engine/Plugins/Experimental/LandscapePatch/Source/LandscapePatch/Public/LandscapeTexturePatch.h
+++ b/Engine/Plugins/Experimental/LandscapePatch/Source/LandscapePatch/Public/LandscapeTexturePatch.h
@@ -242,7 +242,6 @@ public:
 	LANDSCAPEPATCH_API virtual void BlendLayer(UE::Landscape::EditLayers::FRenderParams& RenderParams, UE::Landscape::FRDGBuilderRecorder& RDGBuilderRecorder) override;
 
 	// ULandscapePatchComponent
-	LANDSCAPEPATCH_API virtual UTextureRenderTarget2D* RenderLayer_Native(const FLandscapeBrushParameters& InParameters, const FTransform& HeightmapToWorld) override;
 	LANDSCAPEPATCH_API virtual bool CanAffectHeightmap() const override;
 	LANDSCAPEPATCH_API virtual bool CanAffectWeightmap() const override;
 	LANDSCAPEPATCH_API virtual bool CanAffectWeightmapLayer(const FName& InLayerName) const override;
@@ -668,10 +667,10 @@ private:
 	void TransitionHeightSourceModeInternal(ELandscapeTexturePatchSourceMode OldMode, ELandscapeTexturePatchSourceMode NewMode);
 	FLandscapeHeightPatchConvertToNativeParams GetHeightConvertToNativeParams() const;
 
-	UTextureRenderTarget2D* ApplyToHeightmap(bool bInPerformBlending, UE::Landscape::EditLayers::FRenderParams* RenderParams, UE::Landscape::FRDGBuilderRecorder& RDGBuilderRecorder, UTextureRenderTarget2D* InCombinedResult,
+	UTextureRenderTarget2D* ApplyToHeightmap(UE::Landscape::EditLayers::FRenderParams* RenderParams, UE::Landscape::FRDGBuilderRecorder& RDGBuilderRecorder, UTextureRenderTarget2D* InCombinedResult,
 		const FTransform& LandscapeHeightmapToWorld, bool& bHasRenderedSomething, ERHIAccess OutputAccess = ERHIAccess::None);
 
-	void ApplyToWeightmap(bool bInPerformBlending, UE::Landscape::EditLayers::FRenderParams* RenderParams, UE::Landscape::FRDGBuilderRecorder& RDGBuilderRecorder, ULandscapeWeightPatchTextureInfo* PatchInfo,
+	void ApplyToWeightmap(UE::Landscape::EditLayers::FRenderParams* RenderParams, UE::Landscape::FRDGBuilderRecorder& RDGBuilderRecorder, ULandscapeWeightPatchTextureInfo* PatchInfo,
 		FTextureResource* InMergedLandscapeTextureResource, int32 LandscapeTextureSliceIndex, const FIntPoint& LandscapeTextureResolution, 
 		const FTransform& LandscapeHeightmapToWorld, bool& bHasRenderedSomething, ERHIAccess OutputAccess = ERHIAccess::None);
 
diff --git a/Engine/Shaders/Private/LandscapeLayersPS.usf b/Engine/Shaders/Private/LandscapeLayersPS.usf
--- a/Engine/Shaders/Private/LandscapeLayersPS.usf
+++ b/Engine/Shaders/Private/LandscapeLayersPS.usf
@@ -4,393 +4,10 @@
 #include "Landscape/LandscapeCommon.ush"
 
 Texture2D ReadTexture1;
-Texture2D ReadTexture2; // This one is optional, it's valid when OutputConfig.z == 1
 SamplerState ReadTexture1Sampler;
-SamplerState ReadTexture2Sampler;
-float4 LayerInfo; // x == Layer Alpha value, y == Layer Visibility flag, z == Layer Blend Mode (0: Additive Blend Mode, 1: Alpha Blend Mode), w == Unused
-float4 OutputConfig; // x == ApplyLayerModifiers, y == Alpha Blend Mode or Output weightmap as substractive/additive, z == Use ReadTexture2 and out Delta with ReadTexture1, w == Output Normals (heightmaps) or Unused (weightmaps)
-float2 TextureSize; // x == ReadTexture1 width, y == ReadTexture1 height
-float3 LandscapeGridScale; // x == LS Actor DrawScale.X, y == LS Actor DrawScale.y, z == LS Actor DrawScale.z / 128.0f (ZSCALE)
-float CurrentMipComponentVertexCount;
-float2 CurrentMipTextureSize;
-float2 ParentMipTextureSize;
 float4 SourceOffsetAndSizeUV;
 uint ChannelSwizzleMask;
 
-float ExtractAndScaleHeight(float2 InPackedHeight)
-{
-	float Height = UnpackHeight(InPackedHeight) - 32768.0;
-	return Height * LandscapeGridScale.z * TERRAIN_ZSCALE;
-}
-
-float ExtractAlpha(float2 InPackedAlpha)
-{
-	int AlphaWithFlags = ((int) round(InPackedAlpha.r * 255.0) << 8) | (int) round(InPackedAlpha.g * 255.0);
-	int Alpha = AlphaWithFlags & 0xFFFC; // Last 2 bits are used to store lower/raise flags
-	return float(Alpha) / 65532.0; // 0xFFFC in float
-}
-
-float4 CalculateNormalsFromHeights(float2 InTextureCoordinates, float4 OutColor)
-{
-	float2 TexelSize = 1.0 / TextureSize;
-
-	float MaskMinBorderX = (InTextureCoordinates.x <= TexelSize.x) ? 0.0f : 1.0f; // first pixel X
-	float MaskMinBorderY = (InTextureCoordinates.y <= TexelSize.y) ? 0.0f : 1.0f; // first pixel Y
-	float MaskMaxBorderX = (InTextureCoordinates.x >= 1.0f - TexelSize.x) ? 0.0f : 1.0f; // last pixel X
-	float MaskMaxBorderY = (InTextureCoordinates.y >= 1.0f - TexelSize.y) ? 0.0f : 1.0f; // last pixel Y
-
-	// The triangle topology is the following (where C = center, T = top, B = bottom, L = left, R = right and Nx the normals we need to interpolate):
-	// TL ------ TT
-	// | \       | \
-	// |  \      |  \
-	// |   \     |   \
-	// | N0 \ N1 | N3 \
-	// |     \   |     \
-	// |      \  |      \
-	// |       \ |       \
-	// LL ------ CC ------ RR
-	//   \       | \       |
-	//    \      |  \      |
-	//     \     |   \     |
-	//      \ N2 | N4 \ N5 |
-	//       \   |     \   |
-	//        \  |      \  |
-	//         \ |       \ |
-	//           BB ------ BR
-	// So we can retrieve all 7 samples using 4 gathers : 
-	// Make sure we gather the 4 pixels in the bottom right direction by sampling at (0.75, 0.75) of the pixel (gather retrieves the 4 samples that would be used for bilinear interpolation):
-
-	float2 GatherLocation = InTextureCoordinates + 0.25f * TexelSize;
-	float4 Red0 = ReadTexture1.GatherRed(ReadTexture1Sampler, GatherLocation - TexelSize);
-	float4 Green0 = ReadTexture1.GatherGreen(ReadTexture1Sampler, GatherLocation - TexelSize);
-	float4 Mask0 = ReadTexture1.GatherAlpha(ReadTexture1Sampler, GatherLocation - TexelSize);
-	float4 Red1 = ReadTexture1.GatherRed(ReadTexture1Sampler, GatherLocation);
-	float4 Green1 = ReadTexture1.GatherGreen(ReadTexture1Sampler, GatherLocation);
-	float4 Mask1 = ReadTexture1.GatherAlpha(ReadTexture1Sampler, GatherLocation);
-
-	// [x,y,z] = World Position [w] = valid mask
-	float4 TL = float4((InTextureCoordinates * TextureSize + float2(-1, -1)) * LandscapeGridScale.xy, ExtractAndScaleHeight(float2(Red0.w, Green0.w)), Mask0.w);
-	float4 TT = float4((InTextureCoordinates * TextureSize + float2(+0, -1)) * LandscapeGridScale.xy, ExtractAndScaleHeight(float2(Red0.z, Green0.z)), Mask0.z);
-	float4 CC = float4((InTextureCoordinates * TextureSize + float2(+0, +0)) * LandscapeGridScale.xy, ExtractAndScaleHeight(float2(Red0.y, Green0.y)), Mask0.y);
-	float4 LL = float4((InTextureCoordinates * TextureSize + float2(-1, +0)) * LandscapeGridScale.xy, ExtractAndScaleHeight(float2(Red0.x, Green0.x)), Mask0.x);
-	float4 RR = float4((InTextureCoordinates * TextureSize + float2(+1, +0)) * LandscapeGridScale.xy, ExtractAndScaleHeight(float2(Red1.z, Green1.z)), Mask1.z);
-	float4 BR = float4((InTextureCoordinates * TextureSize + float2(+1, +1)) * LandscapeGridScale.xy, ExtractAndScaleHeight(float2(Red1.y, Green1.y)), Mask1.y);
-	float4 BB = float4((InTextureCoordinates * TextureSize + float2(+0, +1)) * LandscapeGridScale.xy, ExtractAndScaleHeight(float2(Red1.x, Green1.x)), Mask1.x);
-
-	// mask out samples that are off the edge of the texture
-	TL.w *= MaskMinBorderX * MaskMinBorderY;
-	TT.w *= MaskMinBorderY;
-	// CC should never be off the edge
-	LL.w *= MaskMinBorderX;
-	RR.w *= MaskMaxBorderX;
-	BR.w *= MaskMaxBorderX * MaskMaxBorderY;
-	BB.w *= MaskMaxBorderY;
-
-	float3 N0 = ComputeNullableTriangleNormal(CC, LL, TL);
-	float3 N1 = ComputeNullableTriangleNormal(TL, TT, CC);
-	float3 N2 = ComputeNullableTriangleNormal(LL, CC, BB);
-	float3 N3 = ComputeNullableTriangleNormal(RR, CC, TT);
-	float3 N4 = ComputeNullableTriangleNormal(BR, BB, CC);
-	float3 N5 = ComputeNullableTriangleNormal(CC, RR, BR);
-
-	// average the normals
-	float3 FinalNormal = normalize(N0 + N1 + N2 + N3 + N4 + N5);
-
-	// Scale back to be 0-1 normal
-	OutColor.ba = (FinalNormal.xy + 1.0) * 0.5;
-
-	return OutColor;
-}
-
-void PSHeightmapMain(float2 InTextureCoordinates : TEXCOORD0, out float4 OutColor : SV_Target0)
-{
-	float LayerVisibility = LayerInfo.y;
-	float4 SourceColor = ReadTexture1.Sample(ReadTexture1Sampler, InTextureCoordinates);
-	float Height = UnpackHeight(SourceColor.rg);
-	
-	// Perform calculation 0 based
-	Height -= 32768.0f;
-
-	// Output as Layer, so apply Layer info
-	if (OutputConfig.x == 1.0)
-	{
-		float LayerAlpha = LayerInfo.x;
-		Height *= LayerAlpha * LayerVisibility;
-	}
-
-	// Output using 2nd heightmap
-	if (OutputConfig.z == 1.0)
-	{
-		float2 BaseColor = ReadTexture2.Sample(ReadTexture2Sampler, InTextureCoordinates).rg;
-		float BaseHeight = UnpackHeight(BaseColor.rg);
-		BaseHeight -= 32768.0;
-
-		float LayerBlendMode = LayerInfo.z;
-		if (LayerBlendMode == 1 && LayerVisibility == 1)
-		{
-			// Pre-multiplied Alpha Blend Mode
-			float HeightAlphaBlend = ExtractAlpha(SourceColor.ba);
-			float NewHeight = Height + (HeightAlphaBlend * BaseHeight);
-			uint RaiseLower = floor(SourceColor.a * 255.f + 0.5f);
-			bool bLowerTerrain = (RaiseLower & 1);
-			bool bRaiseTerrain = (RaiseLower & 2);
-			if ((bRaiseTerrain && NewHeight > BaseHeight) || (bLowerTerrain && NewHeight < BaseHeight))
-			{
-				Height = NewHeight;
-			}
-			else
-			{
-				Height = BaseHeight;
-			}
-		}
-		else
-		{
-			// Additive Blend Mode
-			Height += BaseHeight;
-		}
-
-		// clamp to valid range
-		Height = clamp(Height, -32768.0, 32767.0);
-	}
-
-	// convert back to positive range
-	Height += 32768.0f;
-
-	float2 PackedHeight = PackHeight(Height);
-
-	// prepare the output color
-	OutColor.r = PackedHeight.x;
-	OutColor.g = PackedHeight.y;
-	OutColor.b = SourceColor.b;
-	OutColor.a = SourceColor.a;
-
-	// we override alpha to indicate valid pixels before calculating normals
-	bool bSetAlphaOne = (OutputConfig.y != 0.0);
-	if (bSetAlphaOne)
-	{
-		OutColor.b = 0.0f;
-		OutColor.a = 1.0f;
-	}
-
-	// Output normals
-	bool bGenerateNormals = (OutputConfig.w >= 0.5);
-	if (bGenerateNormals)
-	{
-		OutColor = CalculateNormalsFromHeights(InTextureCoordinates.xy, OutColor);
-	}
-}
-
-void PSHeightmapMainMips(float2 InTextureCoordinates : TEXCOORD0, out float4 OutColor : SV_Target0)
-{
-	bool IsMinBorderTexelX = false;
-	bool IsMinBorderTexelY = false;
-	bool IsMaxBorderTexelX = false;
-	bool IsMaxBorderTexelY = false;
-
-	// Special case of 1 texel size component
-	if (CurrentMipComponentVertexCount == 1)
-	{
-		if (CurrentMipTextureSize.x >= CurrentMipTextureSize.y) // x size 1 texel
-		{
-			IsMinBorderTexelY = true;
-			IsMinBorderTexelX = (InTextureCoordinates.x * CurrentMipTextureSize.x) <= 1.0;
-			IsMaxBorderTexelX = (InTextureCoordinates.x * CurrentMipTextureSize.x) >= CurrentMipTextureSize.x - 1.0;
-		}
-		else
-		{
-			IsMinBorderTexelX = true;
-			IsMinBorderTexelY = (InTextureCoordinates.y * CurrentMipTextureSize.y) <= 1.0;
-			IsMaxBorderTexelY = (InTextureCoordinates.y * CurrentMipTextureSize.y) >= CurrentMipTextureSize.y - 1.0;
-		}
-	}
-	else
-	{
-		float2 CurrentMipQuadTexelSize = 1.0 / CurrentMipComponentVertexCount;
-		float2 LandscapeQuadUV = frac(InTextureCoordinates * CurrentMipTextureSize / CurrentMipComponentVertexCount);
-
-		IsMinBorderTexelX = LandscapeQuadUV.x <= CurrentMipQuadTexelSize.x;
-		IsMinBorderTexelY = LandscapeQuadUV.y <= CurrentMipQuadTexelSize.y;
-		IsMaxBorderTexelX = LandscapeQuadUV.x >= (CurrentMipQuadTexelSize.x * max(CurrentMipComponentVertexCount - 1, 1.0));
-		IsMaxBorderTexelY = LandscapeQuadUV.y >= (CurrentMipQuadTexelSize.y * max(CurrentMipComponentVertexCount - 1, 1.0));
-	}
-
-	float2 ParentMipTexelSize = 1.0 / ParentMipTextureSize;
-	float2 SourceUV = InTextureCoordinates - (ParentMipTexelSize * 0.5);
-	float HeightResult = 0;
-	float2 NormalResult = 0;
-	float Alpha = 0.5;
-	OutColor = 0; // Default to 0
-
-	if (IsMinBorderTexelX) // on left border
-	{
-		float4 SourceSample = ReadTexture1.SampleLevel(ReadTexture1Sampler, SourceUV, 0);
-		float4 DownSample = ReadTexture1.SampleLevel(ReadTexture1Sampler, SourceUV + float2(0.0, ParentMipTexelSize.y), 0);
-		float SourceHeight = UnpackHeight(SourceSample.rg);
-		float DownHeight = UnpackHeight(DownSample.rg);
-
-		if (IsMinBorderTexelY) // on top border
-		{
-			HeightResult = SourceHeight;
-			NormalResult = float2(SourceSample.ba);
-		}
-		else if (IsMaxBorderTexelY) // on bottom border
-		{
-			HeightResult = DownHeight;
-			NormalResult = float2(DownSample.ba);
-		}
-		else
-		{
-			HeightResult = clamp(round(lerp(SourceHeight, DownHeight, Alpha)), 0.0, 65535.0);
-			NormalResult.x = lerp(SourceSample.b, DownSample.b, Alpha);
-			NormalResult.y = lerp(SourceSample.a, DownSample.a, Alpha);
-		}
-	}
-	else if (IsMaxBorderTexelX) // on right border
-	{
-		float4 RightSample = ReadTexture1.SampleLevel(ReadTexture1Sampler, SourceUV + float2(ParentMipTexelSize.x, 0.0), 0);
-		float4 DownRightSample = ReadTexture1.SampleLevel(ReadTexture1Sampler, SourceUV + ParentMipTexelSize, 0);
-		float RightHeight = UnpackHeight(RightSample.rg);
-		float DownRightHeight = UnpackHeight(DownRightSample.rg);
-
-		if (IsMinBorderTexelY) // on top border
-		{
-			HeightResult = RightHeight;
-			NormalResult = float2(RightSample.ba);
-		}
-		else if (IsMaxBorderTexelY) // on bottom border
-		{
-			HeightResult = DownRightHeight;
-			NormalResult = float2(DownRightSample.ba);
-		}
-		else
-		{
-			HeightResult = clamp(round(lerp(RightHeight, DownRightHeight, Alpha)), 0.0, 65535.0);;
-			NormalResult.x = lerp(RightSample.b, DownRightSample.b, Alpha);
-			NormalResult.y = lerp(RightSample.a, DownRightSample.a, Alpha);
-		}
-	}
-	else // center texel
-	{
-		float4 SourceSample = ReadTexture1.SampleLevel(ReadTexture1Sampler, SourceUV, 0);
-		float4 DownSample = ReadTexture1.SampleLevel(ReadTexture1Sampler, SourceUV + float2(0.0, ParentMipTexelSize.y), 0);
-		float4 RightSample = ReadTexture1.SampleLevel(ReadTexture1Sampler, SourceUV + float2(ParentMipTexelSize.x, 0.0), 0);
-		float4 DownRightSample = ReadTexture1.SampleLevel(ReadTexture1Sampler, SourceUV + ParentMipTexelSize, 0);
-
-		float SourceHeight = UnpackHeight(SourceSample.rg);
-		float DownHeight = UnpackHeight(DownSample.rg);
-		float RightHeight = UnpackHeight(RightSample.rg);
-		float DownRightHeight = UnpackHeight(DownRightSample.rg);
-
-		if (IsMinBorderTexelY) // on top border
-		{
-			HeightResult = clamp(round(lerp(SourceHeight, RightHeight, Alpha)), 0.0, 65535.0);
-			NormalResult.x = lerp(SourceSample.b, RightSample.b, Alpha);
-			NormalResult.y = lerp(SourceSample.a, RightSample.a, Alpha);
-		}
-		else if (IsMaxBorderTexelY) // on bottom border
-		{
-			HeightResult = clamp(round(lerp(DownHeight, DownRightHeight, Alpha)), 0.0, 65535.0);
-			NormalResult.x = lerp(DownSample.b, DownRightSample.b, Alpha);
-			NormalResult.y = lerp(DownSample.a, DownRightSample.a, Alpha);
-		}
-		else
-		{
-			float SourceToRightHeight = lerp(SourceHeight, RightHeight, Alpha);
-			float DownToDownRightHeight = lerp(DownHeight, DownRightHeight, Alpha);
-
-			HeightResult = clamp(round(lerp(SourceToRightHeight, DownToDownRightHeight, Alpha)), 0.0, 65535.0);
-			NormalResult.x = lerp(lerp(SourceSample.b, RightSample.b, Alpha), lerp(DownSample.b, DownRightSample.b, Alpha), Alpha);
-			NormalResult.y = lerp(lerp(SourceSample.a, RightSample.a, Alpha), lerp(DownSample.a, DownRightSample.a, Alpha), Alpha);
-		}
-	}
-
-	float2 PackedHeightResult = PackHeight(HeightResult);
-	OutColor = float4(PackedHeightResult.xy, NormalResult.xy);
-}
-
-void PSWeightmapMain(float2 InTextureCoordinates : TEXCOORD0, out float4 OutColor : SV_Target0)
-{
-	float4 SourceColor = ReadTexture1.Sample(ReadTexture1Sampler, InTextureCoordinates);
-
-	// Output as Layer, so apply Layer info
-	SourceColor = OutputConfig.x == 1.0 ? SourceColor * LayerInfo.x * LayerInfo.y : SourceColor;
-
-	// Output using 2nd heightmap
-	if (OutputConfig.z == 1.0)
-	{
-		float4 BaseColor = ReadTexture2.Sample(ReadTexture2Sampler, InTextureCoordinates);
-
-		// Apply addtive/substractive logic
-		SourceColor = OutputConfig.y == 0.0f ? clamp(BaseColor + SourceColor, 0.0, 1.0) : clamp(BaseColor - SourceColor, 0.0, 1.0);
-	}
-
-	OutColor = SourceColor;
-}
-
-void PSWeightmapMainMips(float2 InTextureCoordinates : TEXCOORD0, out float4 OutColor : SV_Target0)
-{
-	bool IsMinBorderTexelX = false;
-	bool IsMinBorderTexelY = false;
-	bool IsMaxBorderTexelX = false;
-	bool IsMaxBorderTexelY = false;
-
-	// Special case of 1 texel size component
-	if (CurrentMipComponentVertexCount == 1)
-	{
-		if (CurrentMipTextureSize.x >= CurrentMipTextureSize.y) // x size 1 texel
-		{
-			IsMinBorderTexelY = true;
-			IsMinBorderTexelX = (InTextureCoordinates.x * CurrentMipTextureSize.x) <= 1.0;
-			IsMaxBorderTexelX = (InTextureCoordinates.x * CurrentMipTextureSize.x) >= CurrentMipTextureSize.x - 1.0;
-		}
-		else
-		{
-			IsMinBorderTexelX = true;
-			IsMinBorderTexelY = (InTextureCoordinates.y * CurrentMipTextureSize.y) <= 1.0;
-			IsMaxBorderTexelY = (InTextureCoordinates.y * CurrentMipTextureSize.y) >= CurrentMipTextureSize.y - 1.0;
-		}
-	}
-	else
-	{
-		float2 CurrentMipQuadTexelSize = 1.0 / CurrentMipComponentVertexCount;
-		float2 LandscapeQuadUV = frac(InTextureCoordinates * CurrentMipTextureSize / CurrentMipComponentVertexCount);
-
-		IsMinBorderTexelX = LandscapeQuadUV.x <= CurrentMipQuadTexelSize.x;
-		IsMinBorderTexelY = LandscapeQuadUV.y <= CurrentMipQuadTexelSize.y;
-		IsMaxBorderTexelX = LandscapeQuadUV.x >= (CurrentMipQuadTexelSize.x * max(CurrentMipComponentVertexCount - 1, 1.0));
-		IsMaxBorderTexelY = LandscapeQuadUV.y >= (CurrentMipQuadTexelSize.y * max(CurrentMipComponentVertexCount - 1, 1.0));
-	}
-
-	float2 ParentMipTexelSize = 1.0 / ParentMipTextureSize;
-	float2 SourceUV = InTextureCoordinates - (ParentMipTexelSize * 0.5);
-	float Alpha = 0.5;
-	OutColor = 0; // Default to 0
-
-	if (IsMinBorderTexelX) // on left border
-	{
-		float4 SourceSample = ReadTexture1.SampleLevel(ReadTexture1Sampler, SourceUV, 0);
-		float4 DownSample = ReadTexture1.SampleLevel(ReadTexture1Sampler, SourceUV + float2(0.0, ParentMipTexelSize.y), 0);
-
-		OutColor = IsMinBorderTexelY ? SourceSample : (IsMaxBorderTexelY ? DownSample : lerp(SourceSample, DownSample, Alpha));
-	}
-	else if (IsMaxBorderTexelX) // on right border
-	{
-		float4 RightSample = ReadTexture1.SampleLevel(ReadTexture1Sampler, SourceUV + float2(ParentMipTexelSize.x, 0.0), 0);
-		float4 DownRightSample = ReadTexture1.SampleLevel(ReadTexture1Sampler, SourceUV + ParentMipTexelSize, 0);
-
-		OutColor = IsMinBorderTexelY ? RightSample : (IsMaxBorderTexelY ? DownRightSample : lerp(RightSample, DownRightSample, Alpha));
-	}
-	else // center texel
-	{
-		float4 SourceSample = ReadTexture1.SampleLevel(ReadTexture1Sampler, SourceUV, 0);
-		float4 DownSample = ReadTexture1.SampleLevel(ReadTexture1Sampler, SourceUV + float2(0.0, ParentMipTexelSize.y), 0);
-		float4 RightSample = ReadTexture1.SampleLevel(ReadTexture1Sampler, SourceUV + float2(ParentMipTexelSize.x, 0.0), 0);
-		float4 DownRightSample = ReadTexture1.SampleLevel(ReadTexture1Sampler, SourceUV + ParentMipTexelSize, 0);
-
-		OutColor = IsMinBorderTexelY ? lerp(SourceSample, RightSample, Alpha) : (IsMaxBorderTexelY ? lerp(DownSample, DownRightSample, Alpha) : lerp(lerp(SourceSample, RightSample, Alpha), lerp(DownSample, DownRightSample, Alpha), Alpha));
-	}
-}
-
 void CopyTextureVS(in uint VertexId : SV_VertexID, out float4 OutPosition : SV_POSITION, out float2 OutTexCoord : TEXCOORD0)
 {
 	OutTexCoord = float2(((VertexId + 1) / 3) & 1, VertexId & 1);
diff --git a/Engine/Source/Editor/LandscapeEditor/Private/LandscapeEditorDetailCustomization_TargetLayers.cpp b/Engine/Source/Editor/LandscapeEditor/Private/LandscapeEditorDetailCustomization_TargetLayers.cpp
--- a/Engine/Source/Editor/LandscapeEditor/Private/LandscapeEditorDetailCustomization_TargetLayers.cpp
+++ b/Engine/Source/Editor/LandscapeEditor/Private/LandscapeEditorDetailCustomization_TargetLayers.cpp
@@ -260,19 +260,6 @@ void FLandscapeEditorDetailCustomization_TargetLayers::CustomizeDetails(IDetailL
 				.Text(LOCTEXT("FilteredTargetLayers_Tip", "All target layers assigned to this landscape are currently filtered. Use the buttons and/or the filter above to un-hide them."))
 				.AutoWrapText(true)
 		];
-
-	TargetsCategory.AddCustomRow(FText())
-		.Visibility(TAttribute<EVisibility>::Create(TAttribute<EVisibility>::FGetter::CreateStatic(&FLandscapeEditorDetailCustomization_TargetLayers::GetAdvancedWeightBlendingTipVisibility)))
-		[
-			SNew(SMultiLineEditableTextBox)
-				.IsReadOnly(true)
-				.Font(DetailBuilder.GetDetailFontBold())
-				.BackgroundColor(FAppStyle::GetColor("ErrorReporting.BackgroundColor"))
-				.Text(LOCTEXT("AdvancedWeightBlending_Tip", "Some target layers assigned to this landscape are currently using the Advanced Weight Blending the blend method, while edit layers are "
-					"not active or the landscape edit layers merge technique is not set to \"batched merge\". This blend technique is only available when edit layers are active and when batched "
-					"merge is in use (console variable : landscape.EditLayersLocalMerge.Enable 2). Please enable this, otherwise those layers will be treated as Legacy Weight Blending."))
-				.AutoWrapText(true)
-		];
 }
 END_SLATE_FUNCTION_BUILD_OPTIMIZATION
 
@@ -402,28 +389,6 @@ EVisibility FLandscapeEditorDetailCustomization_TargetLayers::GetFilteredTargetL
 	return EVisibility::Collapsed;
 }
 
-// TODO [jonathan.bard] : remove entirely once batched merge is phased out
-EVisibility FLandscapeEditorDetailCustomization_TargetLayers::GetAdvancedWeightBlendingTipVisibility()
-{
-	if (FEdModeLandscape* LandscapeEdMode = GetEditorMode(); (LandscapeEdMode != nullptr) && LandscapeEdMode->CurrentToolTarget.LandscapeInfo.IsValid())
-	{
-		if (ALandscape* LandscapeActor = LandscapeEdMode->CurrentToolTarget.LandscapeInfo->LandscapeActor.Get(); (LandscapeActor != nullptr) && (LandscapeEdMode->CurrentToolTarget.TargetType == ELandscapeToolTargetType::Weightmap))
-		{
-			const TArray<TSharedRef<FLandscapeTargetListInfo>> TargetDisplayList = FLandscapeEditorCustomNodeBuilder_TargetLayers::PrepareTargetLayerList(/*bInSort =*/ false, /*bInFilter = */false);
-			const bool bAtLeastOneLayerUsesPremultipliedAlphaBlending = Algo::AnyOf(TargetDisplayList, [](const TSharedRef<FLandscapeTargetListInfo>& InTargetListInfo)
-				{
-					return (InTargetListInfo->LayerInfoObj != nullptr) && (InTargetListInfo->LayerInfoObj->GetBlendMethod() == ELandscapeTargetLayerBlendMethod::PremultipliedAlphaBlending);
-				});
-			if (bAtLeastOneLayerUsesPremultipliedAlphaBlending && (LandscapeActor->GetEditLayersMergeMode() != ELandscapeEditLayersMergeMode::BatchedMerge))
-			{
-				return EVisibility::Visible;
-			}
-		}
-	}
-
-	return EVisibility::Collapsed;
-}
-
 
 //////////////////////////////////////////////////////////////////////////
 
@@ -1234,18 +1199,18 @@ TSharedPtr<SWidget> FLandscapeEditorCustomNodeBuilder_TargetLayers::GenerateRow(
 					.Padding(FMargin(2))
 					[
 						SNew(SLandscapeAssetThumbnail, Target->bIsLayerReferencedByMaterial ? Target->ThumbnailMIC.Get() : nullptr, ThumbnailPool, Target->bIsLayerReferencedByMaterial ? FName("LandscapeEditor.Target_Weightmap") : FName("LandscapeEditor.Target_Unknown"))
-								.Visibility_Static(&FLandscapeEditorCustomNodeBuilder_TargetLayers::GetDebugModeLayerUsageVisibility_Invert, Target)
-								.ThumbnailSize(FIntPoint(48, 48))
-								// Open landscape layer info asset on double-click on the thumbnail : 
-								.OnAccessAsset_Lambda([Target](UObject* InObject)
-								{ 
-									// Note : the object being returned here is the landscape MIC so it's not what we use for opening the landscape layer info asset : 
-									if ((Target->TargetType == ELandscapeToolTargetType::Weightmap) && (Target->LayerInfoObj != nullptr))
-									{
-										UAssetEditorSubsystem* AssetEditorSubsystem = GEditor->GetEditorSubsystem<UAssetEditorSubsystem>();
-										return AssetEditorSubsystem->OpenEditorForAsset(Target->LayerInfoObj.Get());
-									}
-									return false;
+						.Visibility_Static(&FLandscapeEditorCustomNodeBuilder_TargetLayers::GetDebugModeLayerUsageVisibility_Invert, Target)
+						.ThumbnailSize(FIntPoint(48, 48))
+						// Open landscape layer info asset on double-click on the thumbnail : 
+						.OnAccessAsset_Lambda([Target](UObject* InObject)
+						{ 
+							// Note : the object being returned here is the landscape MIC so it's not what we use for opening the landscape layer info asset : 
+							if ((Target->TargetType == ELandscapeToolTargetType::Weightmap) && (Target->LayerInfoObj != nullptr))
+							{
+								UAssetEditorSubsystem* AssetEditorSubsystem = GEditor->GetEditorSubsystem<UAssetEditorSubsystem>();
+								return AssetEditorSubsystem->OpenEditorForAsset(Target->LayerInfoObj.Get());
+							}
+							return false;
 						})
 					]
 					+ SHorizontalBox::Slot()
@@ -1895,8 +1860,8 @@ void FLandscapeEditorCustomNodeBuilder_TargetLayers::OnTargetLayerSetObject(cons
 			LandscapeEdMode->SetCurrentTargetLayer(Target->LayerName, nullptr);
 			LandscapeInfo->UpdateLayerInfoMap();
 			LandscapeEdMode->UpdateTargetList();
-		return;
-	}
+			return;
+		}
 	}
 
 	ULandscapeLayerInfoObject* SelectedLayerInfo = const_cast<ULandscapeLayerInfoObject*>(CastChecked<ULandscapeLayerInfoObject>(Object));
@@ -2086,9 +2051,9 @@ EVisibility FLandscapeEditorCustomNodeBuilder_TargetLayers::GetLayersSubstractiv
 	if (LandscapeEdMode != nullptr && Target->TargetType == ELandscapeToolTargetType::Weightmap && Target->LayerInfoObj.IsValid())
 	{
 		if (ULandscapeEditLayerBase* EditLayer = LandscapeEdMode->GetEditLayer(LandscapeEdMode->GetSelectedEditLayerIndex()))
-	{
-		return EVisibility::Visible;
-	}
+		{
+			return EVisibility::Visible;
+		}
 	}
 
 	return EVisibility::Collapsed;
@@ -2104,7 +2069,7 @@ ECheckBoxState FLandscapeEditorCustomNodeBuilder_TargetLayers::IsLayersSubstract
 			const bool* AllocationBlend = EditLayer->GetWeightmapLayerAllocationBlend().Find(Target.Get().LayerInfoObj.Get());
 
 			if (AllocationBlend != nullptr)
-	{
+			{
 				return (*AllocationBlend) ? ECheckBoxState::Checked : ECheckBoxState::Unchecked;
 			}
 		}
diff --git a/Engine/Source/Editor/LandscapeEditor/Private/LandscapeEditorDetailCustomization_TargetLayers.h b/Engine/Source/Editor/LandscapeEditor/Private/LandscapeEditorDetailCustomization_TargetLayers.h
--- a/Engine/Source/Editor/LandscapeEditor/Private/LandscapeEditorDetailCustomization_TargetLayers.h
+++ b/Engine/Source/Editor/LandscapeEditor/Private/LandscapeEditorDetailCustomization_TargetLayers.h
@@ -60,7 +60,6 @@ protected:
 	static EVisibility GetPopulateTargetLayersInfoTipVisibility();
 	static EVisibility GetTargetLayersInvalidInfoAssetTipVisibility();
 	static EVisibility GetFilteredTargetLayersListInfoTipVisibility();
-	static EVisibility GetAdvancedWeightBlendingTipVisibility();
 };
 
 class FLandscapeEditorCustomNodeBuilder_TargetLayers : public IDetailCustomNodeBuilder, public TSharedFromThis<FLandscapeEditorCustomNodeBuilder_TargetLayers>
diff --git a/Engine/Source/Runtime/Landscape/Classes/Landscape.h b/Engine/Source/Runtime/Landscape/Classes/Landscape.h
--- a/Engine/Source/Runtime/Landscape/Classes/Landscape.h
+++ b/Engine/Source/Runtime/Landscape/Classes/Landscape.h
@@ -133,8 +133,6 @@ struct FLandscapeLayerBrush
 	virtual ~FLandscapeLayerBrush() = default;
 
 #if WITH_EDITOR
-	UTextureRenderTarget2D* RenderLayer(const FLandscapeBrushParameters& InParameters);
-
 	LANDSCAPE_API ALandscapeBlueprintBrushBase* GetBrush() const;
 	bool AffectsHeightmap() const;
 	bool AffectsWeightmapLayer(const FName& InWeightmapLayerName) const;
@@ -215,7 +213,8 @@ public:
 	TObjectPtr<ULandscapeEditLayerBase> EditLayer;
 };
 
-enum class ELandscapeEditLayersMergeMode : uint8
+
+enum class UE_DEPRECATED(5.7, "Edit layers are now fully processed via BatchedMerge") ELandscapeEditLayersMergeMode : uint8
 {
 	GlobalMerge = 0,
 	LocalMerge, 
@@ -328,8 +327,6 @@ public:
 #endif // WITH_EDITOR
 
 	virtual void PostLoad() override;
-	virtual void BeginDestroy() override;
-	virtual void FinishDestroy() override;
 	//~ End UObject Interface
 
 	/** Computes & returns bounds containing all currently loaded landscape proxies (if any) or this landscape's bounds otherwise */
@@ -468,7 +465,7 @@ public:
 	LANDSCAPE_API uint8 GetLayerCount() const;
 	UE_DEPRECATED(5.6, "Use GetLayersConst")
 	TArrayView<const FLandscapeLayer> GetLayers() const { return MakeArrayView(LandscapeEditLayers); }
-
+	
 	/**
 	 * Runs the given function on each edit layer, with the possibility of early exit
 	 * Most easily used with a lambda as follows:
@@ -487,7 +484,7 @@ public:
 
 PRAGMA_DISABLE_DEPRECATION_WARNINGS
 	UE_DEPRECATED(5.7, "ELandscapeClearMode is deprecated, use ClearEditLayer with ELandscapeToolTargetType instead.")
-	LANDSCAPE_API void ClearLayer(int32 InLayerIndex, TSet<TObjectPtr<ULandscapeComponent>>*InComponents = nullptr, ELandscapeClearMode InClearMode = ELandscapeClearMode::Clear_All);
+	LANDSCAPE_API void ClearLayer(int32 InLayerIndex, TSet<TObjectPtr<ULandscapeComponent>>* InComponents = nullptr, ELandscapeClearMode InClearMode = ELandscapeClearMode::Clear_All);
 	UE_DEPRECATED(5.7, "ELandscapeClearMode is deprecated, use ClearEditLayer with ELandscapeToolTargetType instead.")
 	LANDSCAPE_API void ClearLayer(const FGuid& InLayerGuid, TSet<TObjectPtr<ULandscapeComponent>>* InComponents = nullptr, ELandscapeClearMode InClearMode = ELandscapeClearMode::Clear_All, bool bMarkPackageDirty = true);
 PRAGMA_ENABLE_DEPRECATION_WARNINGS
@@ -532,7 +529,6 @@ PRAGMA_ENABLE_DEPRECATION_WARNINGS
 	LANDSCAPE_API void OnLayerInfoSplineFalloffModulationChanged(ULandscapeLayerInfoObject* InLayerInfo);
 	LANDSCAPE_API void OnPreSave();
 
-	void ReleaseLayersRenderingResource();
 	void ClearDirtyData(ULandscapeComponent* InLandscapeComponent);
 	
 	LANDSCAPE_API void ConvertNonEditLayerLandscape();
@@ -577,17 +573,11 @@ PRAGMA_ENABLE_DEPRECATION_WARNINGS
 	bool GetIsPendingEditLayerConversion() const;
 	void SetIsPendingEditLayerConversion(bool bInPendingEditLayerConversion);
 
-	UE_INTERNAL LANDSCAPE_API ELandscapeEditLayersMergeMode GetEditLayersMergeMode() const;
-
 private:
 	FLandscapeLayer* GetLayerInternal(int32 InLayerIndex);
 	void OnLayerCreatedInternal(ULandscapeEditLayerBase* EditLayer);
 	ULandscapeEditLayerBase* GetEditLayerInternal(int32 InLayerIndex);
 
-	bool SupportsEditLayersLocalMerge();
-
-	bool CreateLayersRenderingResource(bool bInForceLayerResourceReset);
-	void PrepareEditLayersLocalMergeResources();
 	void UpdateLayersContent(bool bInWaitForStreaming = false, bool bInSkipMonitorLandscapeEdModeChanges = false, bool bIntermediateRender = false, bool bFlushRender = false);
 	bool CanUpdateLayersContent() const;
 	void MonitorShaderCompilation();
@@ -600,11 +590,7 @@ private:
 
 	UE::Landscape::EditLayers::FMergeRenderContext PrepareEditLayersMergeRenderContext(const UE::Landscape::EditLayers::FMergeContext& InMergeContext, const UE::Landscape::EditLayers::FMergeRenderParams& InParams);
 
-	FLandscapeBrushParameters BuildLandscapeBrushParameters(bool bInIsHeightmapMerge, const FIntRect& InRenderAreaExtents, UTextureRenderTarget2D* InRenderTarget, FName InWeightmapLayerName = NAME_None);
-
 	int32 RegenerateLayersHeightmaps(const FUpdateLayersContentContext& InUpdateLayersContentContext);
-	int32 PerformLayersHeightmapsLocalMerge(const FUpdateLayersContentContext& InUpdateLayersContentContext, const FEditLayersHeightmapMergeParams& InMergeParams);
-	int32 PerformLayersHeightmapsGlobalMerge(const FUpdateLayersContentContext& InUpdateLayersContentContext, const FEditLayersHeightmapMergeParams& InMergeParams);
 	int32 PerformLayersHeightmapsBatchedMerge(const FUpdateLayersContentContext& InUpdateLayersContentContext, const FEditLayersHeightmapMergeParams& InMergeParams);
 	void ResolveLayersHeightmapTexture(const FTextureToComponentHelper& MapHelper, const TSet<UTexture2D*>& HeightmapsToResolve, bool bIntermediateRender, TArray<FLandscapeEditLayerComponentReadbackResult>& InOutComponentReadbackResults);
 
@@ -612,8 +598,6 @@ private:
 	bool PerformSelectiveRenderEditLayersHeightmapsRT(const FUpdateLayersContentContext& InUpdateLayersContentContext, const FLandscapeEditLayerRenderHeightParams& InSelectiveRenderParams);
 
 	int32 RegenerateLayersWeightmaps(FUpdateLayersContentContext& InUpdateLayersContentContext);
-	int32 PerformLayersWeightmapsLocalMerge(FUpdateLayersContentContext& InUpdateLayersContentContext, const FEditLayersWeightmapMergeParams& InMergeParams);
-	int32 PerformLayersWeightmapsGlobalMerge(FUpdateLayersContentContext& InUpdateLayersContentContext, const FEditLayersWeightmapMergeParams& InMergeParams);
 	int32 PerformLayersWeightmapsBatchedMerge(FUpdateLayersContentContext& InUpdateLayersContentContext, const FEditLayersWeightmapMergeParams& InMergeParams);
 	bool PerformSelectiveRenderEditLayersWeightmapsCPU(FUpdateLayersContentContext& InUpdateLayersContentContext, const FLandscapeEditLayerRenderWeightParams& InSelectiveRenderParams);
 	void ResolveLayersWeightmapTexture(const FTextureToComponentHelper& MapHelper, const TSet<UTexture2D*>& WeightmapsToResolve, bool bIntermediateRender, TArray<FLandscapeEditLayerComponentReadbackResult>& InOutComponentReadbackResults);
@@ -631,49 +615,13 @@ private:
 	bool PrepareLayersTextureResources(const TArray<FLandscapeLayer>& InLayers, bool bInWaitForStreaming);
 	bool PrepareLayersResources(EShaderPlatform InShaderPlatform, bool bInWaitForStreaming);
 	void InvalidateRVTForTextures(const TSet<UTexture2D*>& InTextures);
-	void PrepareLayersHeightmapsLocalMergeRenderThreadData(const FUpdateLayersContentContext& InUpdateLayersContentContext, const FEditLayersHeightmapMergeParams& InMergeParams, EditLayersHeightmapLocalMerge_RenderThread::FMergeInfo& OutRenderThreadData);
-	void PrepareLayersWeightmapsLocalMergeRenderThreadData(const FUpdateLayersContentContext& InUpdateLayersContentContext, const FEditLayersWeightmapMergeParams& InMergeParams, EditLayersWeightmapLocalMerge_RenderThread::FMergeInfo& OutRenderThreadData);
 
 	void UpdateLayersMaterialInstances(const TArray<ULandscapeComponent*>& InLandscapeComponents);
 
-	void PrepareComponentDataToExtractMaterialLayersCS(const TArray<ULandscapeComponent*>& InLandscapeComponents, const ULandscapeEditLayerBase* InEditLayer, int32 InCurrentWeightmapToProcessIndex, const FIntPoint& InLandscapeBase, class FLandscapeTexture2DResource* InOutTextureData,
-														  TArray<struct FLandscapeLayerWeightmapExtractMaterialLayersComponentData>& OutComponentData, TMap<ULandscapeLayerInfoObject*, int32>& OutLayerInfoObjects);
-	void PrepareComponentDataToPackMaterialLayersCS(int32 InCurrentWeightmapToProcessIndex, const FIntPoint& InLandscapeBase, const TArray<ULandscapeComponent*>& InAllLandscapeComponents, TArray<UTexture2D*>& InOutProcessedWeightmaps,
-													TArray<FLandscapeEditLayerReadback*>& OutProcessedCPUReadBacks, TArray<struct FLandscapeLayerWeightmapPackMaterialLayersComponentData>& OutComponentData);
 	void ReallocateLayersWeightmaps(FUpdateLayersContentContext& InUpdateLayersContentContext, const TArray<ULandscapeLayerInfoObject*>& InBrushRequiredAllocations, 
 		const TMap<ULandscapeComponent*, TArray<ULandscapeLayerInfoObject*>>* InPerComponentAllocations, TSet<ULandscapeComponent*>* InRestrictTextureSharingToComponents);
-	void InitializeLayersWeightmapResources();
-	bool GenerateZeroAllocationPerComponents(const TArray<ALandscapeProxy*>& InAllLandscape, const TMap<ULandscapeLayerInfoObject*, bool>& InWeightmapLayersBlendSubstractive);
-
-	void GenerateLayersRenderQuad(const FIntPoint& InVertexPosition, float InVertexSize, const FVector2D& InUVStart, const FVector2D& InUVSize, TArray<struct FLandscapeLayersTriangle>& OutTriangles) const;
-	void GenerateLayersRenderQuadsAtlas(const FIntPoint& InSectionBase, const FVector2D& InScaleBias, float InSubSectionSizeQuad, const FIntPoint& InReadSize, const FIntPoint& InWriteSize, TArray<struct FLandscapeLayersTriangle>& OutTriangles) const;
-	void GenerateLayersRenderQuadsAtlasToNonAtlas(const FIntPoint& InSectionBase, const FVector2D& InScaleBias, float InSubSectionSizeQuad, const FIntPoint& InReadSize, const FIntPoint& InWriteSize, TArray<struct FLandscapeLayersTriangle>& OutTriangles) const;
-	void GenerateLayersRenderQuadsNonAtlas(const FIntPoint& InSectionBase, const FVector2D& InScaleBias, float InSubSectionSizeQuad, const FIntPoint& InReadSize, const FIntPoint& InWriteSize, TArray<struct FLandscapeLayersTriangle>& OutTriangles) const;
-	void GenerateLayersRenderQuadsNonAtlasToAtlas(const FIntPoint& InSectionBase, const FVector2D& InScaleBias, float InSubSectionSizeQuad, const FIntPoint& InReadSize, const FIntPoint& InWriteSize, TArray<struct FLandscapeLayersTriangle>& OutTriangles) const;
-	void GenerateLayersRenderQuadsMip(const FIntPoint& InSectionBase, const FVector2D& InScaleBias, float InSubSectionSizeQuad, const FIntPoint& InReadSize, const FIntPoint& InWriteSize, uint8 InCurrentMip, TArray<FLandscapeLayersTriangle>& OutTriangles) const;
-
-	void ClearLayersWeightmapTextureResource(const FString& InDebugName, FTextureRenderTargetResource* InTextureResourceToClear) const;
-	void DrawHeightmapComponentsToRenderTarget(const FString& InDebugName, const TArray<ULandscapeComponent*>& InComponentsToDraw, const FIntPoint& InLandscapeBase, UTexture* InHeightmapRTRead, UTextureRenderTarget2D* InOptionalHeightmapRTRead2, UTextureRenderTarget2D* InHeightmapRTWrite, ERTDrawingType InDrawType,
-											   bool InClearRTWrite, struct FLandscapeLayersHeightmapShaderParameters& InShaderParams, uint8 InMipRender = 0) const;
-
-	void DrawWeightmapComponentsToRenderTarget(const FString& InDebugName, const TArray<FIntPoint>& InSectionBaseList, const FVector2f& InScaleBias, TArray<FVector2f>* InScaleBiasPerSection, UTexture* InWeightmapRTRead, UTextureRenderTarget2D* InOptionalWeightmapRTRead2, UTextureRenderTarget2D* InWeightmapRTWrite, ERTDrawingType InDrawType,
-												bool InClearRTWrite, struct FLandscapeLayersWeightmapShaderParameters& InShaderParams, uint8 InMipRender) const;
-
-	void DrawWeightmapComponentsToRenderTarget(const FString& InDebugName, const TArray<ULandscapeComponent*>& InComponentsToDraw, const FIntPoint& InLandscapeBase, UTexture* InWeightmapRTRead, UTextureRenderTarget2D* InOptionalWeightmapRTRead2, UTextureRenderTarget2D* InWeightmapRTWrite, ERTDrawingType InDrawType,
-												bool InClearRTWrite, struct FLandscapeLayersWeightmapShaderParameters& InShaderParams, uint8 InMipRender) const;
-
-	void DrawHeightmapComponentsToRenderTargetMips(const TArray<ULandscapeComponent*>& InComponentsToDraw, const FIntPoint& InLandscapeBase, UTexture* InReadHeightmap, bool InClearRTWrite, struct FLandscapeLayersHeightmapShaderParameters& InShaderParams) const;
-	void DrawWeightmapComponentToRenderTargetMips(const TArray<FVector2f>& InTexturePositionsToDraw, UTexture* InReadWeightmap, bool InClearRTWrite, struct FLandscapeLayersWeightmapShaderParameters& InShaderParams) const;
-
-	void CopyTexturePS(const FString& InSourceDebugName, FTextureResource* InSourceResource, const FString& InDestDebugName, FTextureResource* InDestResource) const;
-
 	void InitializeLayers();
 	
-	void PrintLayersDebugRT(const FString& InContext, UTextureRenderTarget2D* InDebugRT, uint8 InMipRender = 0, bool InOutputHeight = true, bool InOutputNormals = false) const;
-	void PrintLayersDebugTextureResource(const FString& InContext, FTextureResource* InTextureResource, uint8 InMipRender = 0, bool InOutputHeight = true, bool InOutputNormals = false) const;
-	void PrintLayersDebugHeightData(const FString& InContext, const TArray<FColor>& InHeightmapData, const FIntPoint& InDataSize, uint8 InMipRender, bool InOutputNormals = false) const;
-	void PrintLayersDebugWeightData(const FString& InContext, const TArray<FColor>& InWeightmapData, const FIntPoint& InDataSize, uint8 InMipRender) const;
-
 	void UpdateWeightDirtyData(ULandscapeComponent* InLandscapeComponent, UTexture2D const* InWeightmap, FColor const* InOldData, FColor const* InNewData, uint8 InChannel);
 	void OnDirtyWeightmap(FTextureToComponentHelper const& MapHelper, UTexture2D const* InWeightmap, FColor const* InOldData, FColor const* InNewData, int32 InMipLevel, uint8 ChangedChannelsMask);
 	void UpdateHeightDirtyData(ULandscapeComponent* InLandscapeComponent, UTexture2D const* InHeightmap, FColor const* InOldData, FColor const* InNewData);
@@ -743,11 +691,13 @@ public:
 	UPROPERTY()
 	TArray<FLandscapeLayer> LandscapeLayers_DEPRECATED;
 
+	UE_DEPRECATED(5.7, "Edit layers are now fully processed via BatchedMerge")
 	UPROPERTY(Transient)
-	TArray<TObjectPtr<UTextureRenderTarget2D>> HeightmapRTList;
+	TArray<TObjectPtr<UTextureRenderTarget2D>> HeightmapRTList_DEPRECATED;
 
+	UE_DEPRECATED(5.7, "Edit layers are now fully processed via BatchedMerge")
 	UPROPERTY(Transient)
-	TArray<TObjectPtr<UTextureRenderTarget2D>> WeightmapRTList;
+	TArray<TObjectPtr<UTextureRenderTarget2D>> WeightmapRTList_DEPRECATED;
 
 	/** List of textures that are not fully streamed in yet (updated every frame to track textures that have finished streaming in) */
 	UPROPERTY(Transient, DuplicateTransient, TextExportTransient)	
@@ -808,8 +758,6 @@ private:
 
 	bool bWarnedLayerMergeResolution = false;
 
-	ELandscapeEditLayersMergeMode CurrentEditLayersMergeMode = ELandscapeEditLayersMergeMode::GlobalMerge;
-
 	struct FWaitingForResourcesNotificationHelper
 	{
 		void Notify(ALandscape* InLandscape, class FLandscapeNotificationManager* InNotificationManager, ELandscapeNotificationType InNotificationType, const FText& InNotificationText);
@@ -826,18 +774,6 @@ private:
 	FWaitingForResourcesNotificationHelper WaitingForTexturesNotificationHelper;
 	FWaitingForResourcesNotificationHelper WaitingForEditLayerResourcesNotificationHelper;
 
-	// Represent all the resolved paint layer, from all layers blended together (size of the landscape x material layer count)
-	class FLandscapeTexture2DArrayResource* CombinedLayersWeightmapAllMaterialLayersResource;
-	
-	// Represent all the resolved paint layer, from the current layer only (size of the landscape x material layer count)
-	class FLandscapeTexture2DArrayResource* CurrentLayersWeightmapAllMaterialLayersResource;	
-	
-	// Used in extracting the material layers data from layer weightmaps (size of the landscape)
-	class FLandscapeTexture2DResource* WeightmapScratchExtractLayerTextureResource;	
-	
-	// Used in packing the material layer data contained into CombinedLayersWeightmapAllMaterialLayersResource to be set again for each component weightmap (size of the landscape)
-	class FLandscapeTexture2DResource* WeightmapScratchPackLayerTextureResource;
-
 	uint32 LastFlushedLayerUpdateFrame = 0;
 	uint32 LastPrepareTextureResourcesCalled = 0;
 	bool bLastPrepareTextureResourcesResult = false;
diff --git a/Engine/Source/Runtime/Landscape/Classes/LandscapeProxy.h b/Engine/Source/Runtime/Landscape/Classes/LandscapeProxy.h
--- a/Engine/Source/Runtime/Landscape/Classes/LandscapeProxy.h
+++ b/Engine/Source/Runtime/Landscape/Classes/LandscapeProxy.h
@@ -870,8 +870,6 @@ PRAGMA_ENABLE_DEPRECATION_WARNINGS
 	TMap<UTexture2D*, FLandscapeEditLayerReadback*> HeightmapsCPUReadback;
 	TMap<UTexture2D*, FLandscapeEditLayerReadback*> WeightmapsCPUReadback;
 
-	FRenderCommandFence ReleaseResourceFence;
-
 	/** Map of weightmap usage */
 	UPROPERTY(Transient, NonTransactional)
 	TMap<TObjectPtr<UTexture2D>, TObjectPtr<ULandscapeWeightmapUsage>> WeightmapUsageMap;
@@ -1049,8 +1047,6 @@ PRAGMA_ENABLE_DEPRECATION_WARNINGS
 	virtual void UnregisterAllComponents(bool bForReregister = false) override;
 	virtual bool IsLevelBoundsRelevant() const override { return true; }
 
-	virtual void BeginDestroy() override;
-	virtual bool IsReadyForFinishDestroy() override;
 	virtual void FinishDestroy() override;
 
 #if WITH_EDITOR
diff --git a/Engine/Source/Runtime/Landscape/Private/Landscape.cpp b/Engine/Source/Runtime/Landscape/Private/Landscape.cpp
--- a/Engine/Source/Runtime/Landscape/Private/Landscape.cpp
+++ b/Engine/Source/Runtime/Landscape/Private/Landscape.cpp
@@ -76,7 +76,6 @@ Landscape.cpp: Terrain rendering
 #include "GlobalShader.h"
 #include "ShaderParameterStruct.h"
 #include "PixelShaderUtils.h"
-#include "LandscapeEditResources.h"
 #include "Rendering/Texture2DResource.h"
 #include "RenderCaptureInterface.h"
 #include "VisualLogger/VisualLogger.h"
@@ -1705,10 +1704,6 @@ ALandscape::ALandscape(const FObjectInitializer& ObjectInitializer)
 	WasCompilingShaders = false;
 	LayerContentUpdateModes = 0;
 	bSplineLayerUpdateRequested = false;
-	CombinedLayersWeightmapAllMaterialLayersResource = nullptr;
-	CurrentLayersWeightmapAllMaterialLayersResource = nullptr;
-	WeightmapScratchExtractLayerTextureResource = nullptr;
-	WeightmapScratchPackLayerTextureResource = nullptr;
 	bLandscapeLayersAreInitialized = false;
 	bLandscapeLayersForceResourceReset = true;
 	LandscapeEdMode = nullptr;
@@ -3304,6 +3299,23 @@ IMPLEMENT_GLOBAL_SHADER(FLandscapePackRGBAChannelsPS, "/Engine/Private/Landscape
 // Render-thread version of the data / functions we need for the local merge of edit layers : 
 namespace UE::Landscape::Private::RenderMergedTexture_RenderThread
 {
+	// Describes a texture and an associated subregion (for heightmap texture sharing), and, optionally, a single channel (for weightmap texture sharing) :
+	struct FTexture2DResourceSubregion
+	{
+		FTexture2DResourceSubregion() = default;
+
+		FTexture2DResourceSubregion(FTexture2DResource* InTexture, const FIntRect& InSubregion, int32 InChannelIndex = INDEX_NONE)
+			: Texture(InTexture)
+			, Subregion(InSubregion)
+			, ChannelIndex(InChannelIndex)
+		{
+		}
+
+		FTexture2DResource* Texture = nullptr;
+		FIntRect Subregion;
+		int32 ChannelIndex = INDEX_NONE;
+	};
+
 	struct FRenderInfo
 	{
 		// Transform to go from the output render area space ((0,0) in the lower left corner, (1,1) in the upper-right) to the temporary render target space
@@ -3649,7 +3661,7 @@ bool ALandscape::RenderMergedTextureInternal(const FTransform& InRenderAreaWorld
 
 				// Effective area of the texture affecting this component (because of texture sharing) :
 				FIntRect SourceTextureSubregion(SourceTextureOffset, SourceTextureOffset + ComponentSize);
-				MergeTextureRenderInfo.ComponentTexturesToRender.Add(ComponentKey, FTexture2DResourceSubregion(SourceTextureResource->GetTexture2DResource(), SourceTextureSubregion, SourceTextureChannel));
+				MergeTextureRenderInfo.ComponentTexturesToRender.Add(ComponentKey, RenderMergedTexture_RenderThread::FTexture2DResourceSubregion(SourceTextureResource->GetTexture2DResource(), SourceTextureSubregion, SourceTextureChannel));
 
 				// Since this component will be rendered in the render target, we can now expand the render target's bounds :
 				RenderTargetComponentIndicesBoundingRect.Union(FIntRect(ComponentKey, ComponentKey + FIntPoint(1)));
diff --git a/Engine/Source/Runtime/Landscape/Private/LandscapeEditLayerMergeRenderContext.cpp b/Engine/Source/Runtime/Landscape/Private/LandscapeEditLayerMergeRenderContext.cpp
--- a/Engine/Source/Runtime/Landscape/Private/LandscapeEditLayerMergeRenderContext.cpp
+++ b/Engine/Source/Runtime/Landscape/Private/LandscapeEditLayerMergeRenderContext.cpp
@@ -28,7 +28,7 @@
 #include "VisualLogger/VisualLogger.h"
 
 extern TAutoConsoleVariable<float> CVarLandscapeBatchedMergeVisualLogOffsetIncrement;
-extern TAutoConsoleVariable<int32> CVarLandscapeEditLayersClearBeforeEachWriteToScratch;
+extern TAutoConsoleVariable<int32> CVarLandscapeBatchedMergeClearBeforeEachWriteToScratch;
 extern TAutoConsoleVariable<int32> CVarLandscapeBatchedMergeVisualLogShowMergeType;
 extern TAutoConsoleVariable<int32> CVarLandscapeBatchedMergeVisualLogShowMergeProcess;
 extern TAutoConsoleVariable<float> CVarLandscapeBatchedMergeVisualLogAlpha;
@@ -509,7 +509,7 @@ void FMergeRenderContext::CycleBlendRenderTargets(UE::Landscape::FRDGBuilderReco
 	if (!bFirstWrite)
 	{
 		// Optionally clear the write render target for debug purposes : 
-		if (CVarLandscapeEditLayersClearBeforeEachWriteToScratch.GetValueOnGameThread())
+		if (CVarLandscapeBatchedMergeClearBeforeEachWriteToScratch.GetValueOnGameThread())
 		{
 			GetBlendRenderTargetWrite()->Clear(RDGBuilderRecorder);
 		}
