Landscape Patch component : fixed BP function GetLandscapeHeightmapCoordsToWorld

diff --git a/Engine/Plugins/Experimental/LandscapePatch/Source/LandscapePatch/Private/LandscapePatchComponent.cpp b/Engine/Plugins/Experimental/LandscapePatch/Source/LandscapePatch/Private/LandscapePatchComponent.cpp
--- a/Engine/Plugins/Experimental/LandscapePatch/Source/LandscapePatch/Private/LandscapePatchComponent.cpp
+++ b/Engine/Plugins/Experimental/LandscapePatch/Source/LandscapePatch/Private/LandscapePatchComponent.cpp
@@ -880,7 +880,7 @@ FTransform ULandscapePatchComponent::GetLandscapeHeightmapCoordsToWorld() const
 {
 	if (Landscape.IsValid())
 	{
-		UE::Landscape::PatchUtil::GetHeightmapToWorld(Landscape->GetTransform());
+		return UE::Landscape::PatchUtil::GetHeightmapToWorld(Landscape->GetTransform());
 	}
 	return FTransform::Identity;
 }
diff --git a/Engine/Plugins/Experimental/LandscapePatch/Source/LandscapePatch/Public/LandscapePatchComponent.h b/Engine/Plugins/Experimental/LandscapePatch/Source/LandscapePatch/Public/LandscapePatchComponent.h
--- a/Engine/Plugins/Experimental/LandscapePatch/Source/LandscapePatch/Public/LandscapePatchComponent.h
+++ b/Engine/Plugins/Experimental/LandscapePatch/Source/LandscapePatch/Public/LandscapePatchComponent.h
@@ -135,7 +135,7 @@ public:
 #if WITH_EDITOR
 
 	// Cleanup method that can be used to fix the patch's binding to a landscape. An arbitrary landscape
-	// will be bound to if no guid, or landscape is set on the patch.
+	//  will be bound to if no guid, or no landscape is set on the patch.
 	void FixBindings();
 
 	// Allows the patch to update its UI and cached pointer
