LandscapePatch :
* Added  Engine.ini file for plugin : in order to override engine settings
* Marked LandscapeCircleHeightPatch as "internal class" to avoid exposing it to the class viewer

diff --git a/Engine/Plugins/Editor/LandscapePatch/Config/Engine.ini b/Engine/Plugins/Editor/LandscapePatch/Config/Engine.ini
new file mode 100644
--- /dev/null
+++ b/Engine/Plugins/Editor/LandscapePatch/Config/Engine.ini
@@ -0,0 +1,2 @@
+[/Script/ClassViewer.ClassViewerProjectSettings]
++InternalOnlyClasses=/Script/LandscapePatch.LandscapeCircleHeightPatch
\ No newline at end of file
