Fixed typo in comment

#rnx
#skipci

diff --git a/Engine/Plugins/Editor/LandscapePatch/Source/LandscapePatch/Public/LandscapeTexturePatch.h b/Engine/Plugins/Editor/LandscapePatch/Source/LandscapePatch/Public/LandscapeTexturePatch.h
--- a/Engine/Plugins/Editor/LandscapePatch/Source/LandscapePatch/Public/LandscapeTexturePatch.h
+++ b/Engine/Plugins/Editor/LandscapePatch/Source/LandscapePatch/Public/LandscapeTexturePatch.h
@@ -141,8 +141,8 @@ public:
 	double ZeroInEncoding = 0;
 
 	/**
-	 * The scale that should be aplied to the data stored in the patch relative to the zero in the encoding, in world coordinates.
-	 * For instance if the encoding is [0,1], and 0.5 correponds to 0, a WorldSpaceEncoding Scale of 100 means that the resulting
+	 * The scale that should be applied to the data stored in the patch relative to the zero in the encoding, in world coordinates.
+	 * For instance if the encoding is [0,1], and 0.5 corresponds to 0, a WorldSpaceEncoding Scale of 100 means that the resulting
 	 * values will lie in the range [-50, 50] in world space, which would be [-0.5, 0.5] in the landscape local heights if the Z
 	 * scale is 100.
 	 */
