* Hooked up icons for landscape patch component / plugin / edit layer
* Fixed missing class thumbnails for landscape edit layer classes

diff --git a/Engine/Plugins/Editor/LandscapePatch/Source/LandscapePatchEditorOnly/LandscapePatchEditorOnly.Build.cs b/Engine/Plugins/Editor/LandscapePatch/Source/LandscapePatchEditorOnly/LandscapePatchEditorOnly.Build.cs
--- a/Engine/Plugins/Editor/LandscapePatch/Source/LandscapePatchEditorOnly/LandscapePatchEditorOnly.Build.cs
+++ b/Engine/Plugins/Editor/LandscapePatch/Source/LandscapePatchEditorOnly/LandscapePatchEditorOnly.Build.cs
@@ -36,6 +36,7 @@ public class LandscapePatchEditorOnly : ModuleRules
 			{
 				"Engine",
 				"LandscapePatch",
+				"Projects",
 				"Slate",
 				"SlateCore",
 				"UnrealEd",
diff --git a/Engine/Plugins/Editor/LandscapePatch/Source/LandscapePatchEditorOnly/Private/LandscapePatchEditorOnlyModule.cpp b/Engine/Plugins/Editor/LandscapePatch/Source/LandscapePatchEditorOnly/Private/LandscapePatchEditorOnlyModule.cpp
--- a/Engine/Plugins/Editor/LandscapePatch/Source/LandscapePatchEditorOnly/Private/LandscapePatchEditorOnlyModule.cpp
+++ b/Engine/Plugins/Editor/LandscapePatch/Source/LandscapePatchEditorOnly/Private/LandscapePatchEditorOnlyModule.cpp
@@ -3,6 +3,7 @@
 #include "LandscapePatchEditorOnlyModule.h"
 
 #include "Editor/UnrealEdEngine.h"
+#include "LandscapePatchUIStyle.h"
 #include "LandscapeTexturePatch.h"
 #include "LandscapeTexturePatchVisualizer.h"
 #include "Modules/ModuleManager.h"
@@ -13,8 +14,11 @@
 
 void FLandscapePatchEditorOnlyModule::StartupModule()
 {
+	FLandscapePatchUIStyle::Initialize();
+
 	if (GUnrealEd)
 	{
+
 		TSharedPtr<FLandscapeTexturePatchVisualizer> Visualizer = MakeShared<FLandscapeTexturePatchVisualizer>();
 		GUnrealEd->RegisterComponentVisualizer(ULandscapeTexturePatch::StaticClass()->GetFName(), Visualizer);
 		// This call should maybe be inside the RegisterComponentVisualizer call above, but since it's not,
@@ -43,6 +47,8 @@ void FLandscapePatchEditorOnlyModule::ShutdownModule()
 			PropertyEditorModule->UnregisterCustomClassLayout(ClassName);
 		}
 	}
+
+	FLandscapePatchUIStyle::Shutdown();
 }
 
 #undef LOCTEXT_NAMESPACE
diff --git a/Engine/Plugins/Editor/LandscapePatch/Source/LandscapePatchEditorOnly/Private/LandscapePatchUIStyle.cpp b/Engine/Plugins/Editor/LandscapePatch/Source/LandscapePatchEditorOnly/Private/LandscapePatchUIStyle.cpp
new file mode 100644
--- /dev/null
+++ b/Engine/Plugins/Editor/LandscapePatch/Source/LandscapePatchEditorOnly/Private/LandscapePatchUIStyle.cpp
@@ -0,0 +1,63 @@
+// Copyright Epic Games, Inc. All Rights Reserved.
+
+#include "LandscapePatchUIStyle.h"
+#include "Brushes/SlateImageBrush.h"
+#include "Styling/SlateStyleRegistry.h"
+#include "Interfaces/IPluginManager.h"
+#include "Misc/Paths.h"
+#include "Styling/SlateStyle.h"
+#include "Styling/SlateStyleMacros.h"
+
+#define RootToContentDir StyleSet->RootToContentDir
+
+TSharedPtr<FSlateStyleSet> FLandscapePatchUIStyle::StyleSet;
+TSharedPtr<class ISlateStyle> FLandscapePatchUIStyle::Get() { return StyleSet; }
+
+FName FLandscapePatchUIStyle::GetStyleSetName()
+{
+	static FName LandscapePatchStyleName(TEXT("LandscapePatchUIStyle"));
+	return LandscapePatchStyleName;
+}
+
+void FLandscapePatchUIStyle::Initialize()
+{
+	// Const icon sizes
+	const FVector2D Icon8x8(8.0f, 8.0f);
+	const FVector2D Icon16x16(16.0f, 16.0f);
+	const FVector2D Icon20x20(20.0f, 20.0f);
+	const FVector2D Icon40x40(40.0f, 40.0f);
+	const FVector2D Icon64x64(64.0f, 64.0f);
+
+	// Only register once
+	if (StyleSet.IsValid())
+	{
+		return;
+	}
+
+	StyleSet = MakeShareable(new FSlateStyleSet(GetStyleSetName()));
+
+	StyleSet->SetContentRoot(IPluginManager::Get().FindPlugin(TEXT("LandscapePatch"))->GetContentDir());
+	StyleSet->SetCoreContentRoot(FPaths::EngineContentDir() / TEXT("Slate"));
+
+	// LandscapePatchEditLayer
+	StyleSet->Set("ClassIcon.LandscapePatchEditLayer", new IMAGE_BRUSH_SVG("Icons/LandscapePatchEditLayer_16", Icon16x16));
+	StyleSet->Set("ClassThumbnail.LandscapePatchEditLayer", new IMAGE_BRUSH_SVG("Icons/LandscapePatchEditLayer_64", Icon64x64));
+
+	// LandscapePatchComponent
+	StyleSet->Set("ClassIcon.LandscapePatchComponent", new IMAGE_BRUSH_SVG("Icons/LandscapePatchComponent_16", Icon16x16));
+	StyleSet->Set("ClassThumbnail.LandscapePatchComponent", new IMAGE_BRUSH_SVG("Icons/LandscapePatchComponent_64", Icon64x64));
+
+	FSlateStyleRegistry::RegisterSlateStyle(*StyleSet.Get());
+};
+
+#undef RootToContentDir
+
+void FLandscapePatchUIStyle::Shutdown()
+{
+	if (StyleSet.IsValid())
+	{
+		FSlateStyleRegistry::UnRegisterSlateStyle(*StyleSet.Get());
+		ensure(StyleSet.IsUnique());
+		StyleSet.Reset();
+	}
+}
diff --git a/Engine/Plugins/Editor/LandscapePatch/Source/LandscapePatchEditorOnly/Private/LandscapePatchUIStyle.h b/Engine/Plugins/Editor/LandscapePatch/Source/LandscapePatchEditorOnly/Private/LandscapePatchUIStyle.h
new file mode 100644
--- /dev/null
+++ b/Engine/Plugins/Editor/LandscapePatch/Source/LandscapePatchEditorOnly/Private/LandscapePatchUIStyle.h
@@ -0,0 +1,20 @@
+// Copyright Epic Games, Inc. All Rights Reserved.
+
+#pragma once
+
+#include "Templates/SharedPointer.h"
+
+class FLandscapePatchUIStyle
+{
+public:
+	static void Initialize();
+
+	static void Shutdown();
+
+	static TSharedPtr<class ISlateStyle> Get();
+
+	static FName GetStyleSetName();
+
+private:
+	static TSharedPtr<class FSlateStyleSet> StyleSet;
+};
