Landscape Editor - Fixed assert caused by [REDACTED_PERSON] layer on Undo operations
* The flag is true anytime the package is not loaded, as a result, Undo re-registers with bUserTriggered=true
* Modifying the patch (eg. Transform) from Selection mode now Soft-dirties the landscape correctly

diff --git a/Engine/Plugins/Editor/LandscapePatch/Source/LandscapePatch/Private/LandscapePatchEditLayer.cpp b/Engine/Plugins/Editor/LandscapePatch/Source/LandscapePatch/Private/LandscapePatchEditLayer.cpp
--- a/Engine/Plugins/Editor/LandscapePatch/Source/LandscapePatch/Private/LandscapePatchEditLayer.cpp
+++ b/Engine/Plugins/Editor/LandscapePatch/Source/LandscapePatch/Private/LandscapePatchEditLayer.cpp
@@ -344,7 +344,7 @@ void ULandscapePatchEditLayer::RequestLandscapeUpdate(bool bInUserTriggered /* =
 	}
 
 	// TODO: Consider passing a parameter down to say when we're not updating height, only weights, when that is the case
-	OwningLandscape->RequestLayersContentUpdateForceAll(ELandscapeLayerUpdateMode::Update_All, bInUserTriggered || !UE::GetIsEditorLoadingPackage());
+	OwningLandscape->RequestLayersContentUpdateForceAll(ELandscapeLayerUpdateMode::Update_All, bInUserTriggered && !UE::GetIsEditorLoadingPackage());
 }
 
 bool ULandscapePatchEditLayer::ShouldPatchBeIncludedInList(const ULandscapePatchComponent* Patch) const
