Landscape Patch - Added ability to define alpha override settings to the landscape texture patch height and weight data
* Deprecates bUseAlphaChannel for enum type allowing users to define a custom alpha source mode. Height/Weight patches can set custom alpha sources. Adds ability to select texture color channel for the alpha mask
* Introduces default weight patch properties (base value for all weight patches). Adds a default OverrideBlendMode for all weight patches (previously, weight patches would override the height blend mode)
* Added texture channel enum for height / weight texture patches allowing users to define the source texture channel. By default, red texture channel is used for height/weight. Height patches are limited to RG texture channel when InternalSourceTexture/Native encoding is used

#Synced-CL 48916257

diff --git a/Engine/Plugins/Editor/LandscapePatch/Shaders/Private/LandscapeTexturePatchPS.usf b/Engine/Plugins/Editor/LandscapePatch/Shaders/Private/LandscapeTexturePatchPS.usf
--- a/Engine/Plugins/Editor/LandscapePatch/Shaders/Private/LandscapeTexturePatchPS.usf
+++ b/Engine/Plugins/Editor/LandscapePatch/Shaders/Private/LandscapeTexturePatchPS.usf
@@ -89,6 +89,42 @@ float4 SamplePatchTexture(Texture2D<float4> InPatchTexture, SamplerState InPatch
 	}
 	return TextureSample;
 }
+
+float GetChannelFromSample(uint InTextureChannel, float4 InSampledValue)
+{
+	switch (InTextureChannel)
+	{
+		case TEXTURE_CHANNEL_R:
+			return InSampledValue.r;
+		case TEXTURE_CHANNEL_G:
+			return InSampledValue.g;
+		case TEXTURE_CHANNEL_B:
+			return InSampledValue.b;
+		case TEXTURE_CHANNEL_A:
+			return InSampledValue.a;
+		case TEXTURE_CHANNEL_LUM:
+			return Luminance(InSampledValue.rgb);
+		default:
+			return 1;
+	}
+}
+
+float GetAlphaFactor(uint InAlphaMode, uint InAlphaTextureChannel, float4 InPatchTextureSampledValue, float4 InAlphaTextureSampledValue)
+{
+	switch (InAlphaMode)
+	{
+		case ALPHA_PATCH_NONE:
+			return 1;
+		case ALPHA_PATCH_INTERNAL_TEXTURE:
+		case ALPHA_PATCH_RENDER_TARGET:
+		case ALPHA_PATCH_TEXTURE_ASSET:
+			return GetChannelFromSample(InAlphaTextureChannel, InAlphaTextureSampledValue);
+		case ALPHA_PATCH_SRC_TEXTURE_CHANNEL:
+			return GetChannelFromSample(InAlphaTextureChannel, InPatchTextureSampledValue);
+		default:
+			return 1;
+	}
+}
 #endif // APPLY_HEIGHT_PATCH || APPLY_WEIGHT_PATCH
 
 #if APPLY_HEIGHT_PATCH
@@ -108,6 +144,12 @@ float InHeightOffset;
 uint InBlendMode;
 // A combination of flags, whose positions are set in the corresponding cpp file.
 uint InFlags;
+// Height source texture color channel
+uint InTextureChannel;
+// Alpha source mode, override texture (White system texture if undefined), and alpha texture color channel
+uint InAlphaMode;
+Texture2D<float4> InHeightPatchAlphaTexture;
+uint InAlphaTextureChannel;
 
 uint ConvertPatchBlendModeToHeightAlphaFlags()
 {
@@ -131,7 +173,6 @@ void ApplyLandscapeTextureHeightPatch(in float4 SVPos : SV_POSITION, out float4
 	// Unpack our flags. The flag choices are set in FApplyLandscapeTextureHeightPatchPS::ModifyCompilationEnvironment
 	// so that they match to the cpp file
 	bool bRectangularFalloff = InFlags & RECTANGULAR_FALLOFF_FLAG;
-	bool bApplyPatchAlpha = InFlags & APPLY_PATCH_ALPHA_FLAG;
 	bool bInputIsPackedHeight = InFlags & INPUT_IS_PACKED_HEIGHT_FLAG;
 	
 	// We need only the 2D affine transformation that goes from landscape XY heightmap integer coordinates
@@ -141,17 +182,20 @@ void ApplyLandscapeTextureHeightPatch(in float4 SVPos : SV_POSITION, out float4
 	
 	float2 PatchUVCoordinates = mul(HeightmapToPatchRotateScale, SVPos.xy) + HeightmapToPatchTranslate;
 	float4 PatchSampledValue = SamplePatchTexture(InHeightPatch, InHeightPatchSampler, PatchUVCoordinates);
-	float PatchStoredHeight = bInputIsPackedHeight ? UnpackHeight(PatchSampledValue.xy) : PatchSampledValue.x;
+
+	// When using an internal texture/native packed height we enforce that the data comes from the RG channels
+	// Otherwise, get the selected texture channel to use for the height
+	float PatchStoredHeight = bInputIsPackedHeight ? UnpackHeight(PatchSampledValue.xy) : GetChannelFromSample(InTextureChannel, PatchSampledValue);
 	
 	float PatchSignedHeight = InHeightScale * (PatchStoredHeight - InZeroInEncoding) + InHeightOffset;
 	
 	float Alpha = GetFalloffAlpha(InFalloffWorldMargin, InPatchWorldDimensions, PatchUVCoordinates, InEdgeUVDeadBorder, bRectangularFalloff);
 	
-	if (bApplyPatchAlpha)
-	{
-		Alpha *= PatchSampledValue.a;
-	}
-	
+	// TODO [jared.ritchie] Apply alpha specific UV scale/offset/rotation once tiling is added
+	float4 PatchAlphaSampledValue = SamplePatchTexture(InHeightPatchAlphaTexture, InHeightPatchSampler, PatchUVCoordinates);
+	// Apply the alpha from the source height texture or use the standalone alpha texture
+	Alpha *= GetAlphaFactor(InAlphaMode, InAlphaTextureChannel, PatchSampledValue, PatchAlphaSampledValue);
+
 	float OutputHeight = LANDSCAPE_MID_VALUE + PatchSignedHeight;
 	OutColor = float4(PackHeight(OutputHeight), PackHeightAlpha(Alpha, ConvertPatchBlendModeToHeightAlphaFlags()));
 }
@@ -215,6 +259,12 @@ float InFalloffWorldMargin;
 uint InBlendMode;
 // A combination of flags, whose positions are set in the corresponding cpp file.
 uint InFlags;
+// Weight source texture color channel
+uint InTextureChannel;
+// Alpha source mode, override texture (White system texture if undefined), and alpha texture color channel
+uint InAlphaMode;
+Texture2D<float4> InWeightPatchAlphaTexture;
+uint InAlphaTextureChannel;
 
 uint ConvertPatchBlendModeToWeightAlphaFlags()
 {
@@ -238,8 +288,7 @@ void ApplyLandscapeTextureWeightPatch(in float4 SVPos : SV_POSITION, out float4
 	// Unpack our flags. The flag choices are set in FApplyLandscapeTextureWeightPatchPS::ModifyCompilationEnvironment
 	// so that they match to the cpp file
 	bool bRectangularFalloff = InFlags & RECTANGULAR_FALLOFF_FLAG;
-	bool bApplyPatchAlpha = InFlags & APPLY_PATCH_ALPHA_FLAG;
-	
+
 	// We need only the 2D affine transformation that goes from landscape weightmap integer coordinates
 	// to patch UV coordinates.
 	float2x2 WeightmapToPatchRotateScale = (float2x2) InWeightmapToPatch;
@@ -247,14 +296,15 @@ void ApplyLandscapeTextureWeightPatch(in float4 SVPos : SV_POSITION, out float4
 	
 	float2 PatchUVCoordinates = mul(WeightmapToPatchRotateScale, SVPos.xy) + WeightmapToPatchTranslate;
 	float4 PatchSampledValue = SamplePatchTexture(InWeightPatch, InWeightPatchSampler, PatchUVCoordinates);
-	float PatchWeight = PatchSampledValue.x;
-	
+	// Get the selected texture color channel from the sample
+	float PatchWeight = GetChannelFromSample(InTextureChannel, PatchSampledValue);
+
 	float Alpha = GetFalloffAlpha(InFalloffWorldMargin, InPatchWorldDimensions, PatchUVCoordinates, InEdgeUVDeadBorder, bRectangularFalloff);
-	
-	if (bApplyPatchAlpha)
-	{
-		Alpha *= PatchSampledValue.a;
-	}
+
+	// TODO [jared.ritchie] Apply alpha specific UV scale/offset/rotation once tiling is added
+	// Apply the alpha from the source weight texture or use the standalone alpha texture
+	float4 PatchAlphaSampledValue = SamplePatchTexture(InWeightPatchAlphaTexture, InWeightPatchSampler, PatchUVCoordinates);
+	Alpha *= GetAlphaFactor(InAlphaMode, InAlphaTextureChannel, PatchSampledValue, PatchAlphaSampledValue);
 	
 	OutColor = float4(PackWeight(PatchWeight), PackWeightAlpha(Alpha, ConvertPatchBlendModeToWeightAlphaFlags()));
 }
diff --git a/Engine/Plugins/Editor/LandscapePatch/Source/LandscapePatch/Private/LandscapeTexturePatch.cpp b/Engine/Plugins/Editor/LandscapePatch/Source/LandscapePatch/Private/LandscapeTexturePatch.cpp
--- a/Engine/Plugins/Editor/LandscapePatch/Source/LandscapePatch/Private/LandscapeTexturePatch.cpp
+++ b/Engine/Plugins/Editor/LandscapePatch/Source/LandscapePatch/Private/LandscapeTexturePatch.cpp
@@ -25,6 +25,7 @@
 #include "Algo/AnyOf.h"
 #include "Algo/ForEach.h"
 #include "Containers/Ticker.h"
+#include "UObject/FortniteMainBranchObjectVersion.h"
 
 #include UE_INLINE_GENERATED_CPP_BY_NAME(LandscapeTexturePatch)
 
@@ -35,41 +36,50 @@ namespace LandscapeTexturePatchLocals
 #if WITH_EDITOR
 	template <typename TextureBackedRTType>
 	void TransitionSourceMode(ELandscapeTexturePatchSourceMode OldMode, ELandscapeTexturePatchSourceMode NewMode,
-		TObjectPtr<UTexture>& TextureAsset, TObjectPtr<TextureBackedRTType>& InternalData,
+		TObjectPtr<UTexture>& OutTextureAsset, TObjectPtr<TextureBackedRTType>& OutInternalData, ELandscapeTexturePatchTextureChannel& OutTextureChannel,
 		TUniqueFunction<TextureBackedRTType* ()> InternalDataBuilder)
 	{
+		// When transitioning from an invalid source mode enable the red channel (default color channel used by patch)
+		if (OutTextureChannel == ELandscapeTexturePatchTextureChannel::None && NewMode != OldMode)
+		{
+			OutTextureChannel = ELandscapeTexturePatchTextureChannel::Red;
+		}
+
 		if (NewMode == ELandscapeTexturePatchSourceMode::None)
 		{
-			TextureAsset = nullptr;
-			InternalData = nullptr;
+			OutTextureAsset = nullptr;
+			OutInternalData = nullptr;
+
+			// Reset the color channel
+			OutTextureChannel = ELandscapeTexturePatchTextureChannel::None;
 		}
 		else if (NewMode == ELandscapeTexturePatchSourceMode::TextureAsset)
 		{
-			InternalData = nullptr;
+			OutInternalData = nullptr;
 		}
 		else // new mode is internal texture or render target
 		{
 			bool bWillUseTextureOnly = (NewMode == ELandscapeTexturePatchSourceMode::InternalTexture);
 			bool bNeedToCopyTextureAsset = (OldMode == ELandscapeTexturePatchSourceMode::TextureAsset
-				&& IsValid(TextureAsset) && TextureAsset->GetResource());
+				&& IsValid(OutTextureAsset) && OutTextureAsset->GetResource());
 
-			if (!InternalData)
+			if (!OutInternalData)
 			{
-				InternalData = InternalDataBuilder();
-				InternalData->SetUseInternalTextureOnly(bWillUseTextureOnly && !bNeedToCopyTextureAsset);
-				InternalData->Initialize();
+				OutInternalData = InternalDataBuilder();
+				OutInternalData->SetUseInternalTextureOnly(bWillUseTextureOnly && !bNeedToCopyTextureAsset);
+				OutInternalData->Initialize();
 			}
 			else
 			{
-				InternalData->Modify();
+				OutInternalData->Modify();
 			}
 
-			InternalData->SetUseInternalTextureOnly(bWillUseTextureOnly && !bNeedToCopyTextureAsset);
+			OutInternalData->SetUseInternalTextureOnly(bWillUseTextureOnly && !bNeedToCopyTextureAsset);
 			if (bNeedToCopyTextureAsset)
 			{
 				// Copy the currently set texture asset to our render target
-				FTextureResource* Source = TextureAsset->GetResource();
-				FTextureResource* Destination = InternalData->GetRenderTarget()->GetResource();
+				FTextureResource* Source = OutTextureAsset->GetResource();
+				FTextureResource* Destination = OutInternalData->GetRenderTarget()->GetResource();
 
 				ENQUEUE_RENDER_COMMAND(LandscapeTextureHeightPatchRTToTexture)(
 					[Source, Destination](FRHICommandListImmediate& RHICmdList)
@@ -80,13 +90,13 @@ namespace LandscapeTexturePatchLocals
 
 			// Note that the duplicate SetUseInternalTextureOnly calls (in cases where we don't need to copy the texture asset)
 			// are fine because they don't do anything.
-			InternalData->SetUseInternalTextureOnly(bWillUseTextureOnly);
+			OutInternalData->SetUseInternalTextureOnly(bWillUseTextureOnly);
 
-			TextureAsset = nullptr;
+			OutTextureAsset = nullptr;
 		}
 	}
 
-	// TODO: The way we currently do intialization is a bit of a hack in that we actually request to do
+	// TODO: The way we currently do initialization is a bit of a hack in that we actually request to do
 	//  a landscape update but we read instead of writing. In batched merge, this might not always work
 	//  properly because a patch might be at the edge of a rendered batch, and thus only have part of it
 	//  be initialized properly. The proper way to do reinitialization would be to use a special function
@@ -112,6 +122,41 @@ namespace LandscapeTexturePatchLocals
 				"a region of landscape being rendered. Reinitialization might not work be fully supported here."));
 		}
 	}
+
+	void GetTexturePatchAlphaResource(FLandscapeTexturePatchAlphaSettings& InAlphaSettings, UObject* InOuter, FTextureResource*& OutAlphaTextureResource)
+	{
+		UTexture* PatchAlphaUObject = nullptr;
+
+		switch (InAlphaSettings.GetAlphaSourceMode())
+		{
+		case ELandscapeTexturePatchAlphaSourceMode::None:
+			break;
+		case ELandscapeTexturePatchAlphaSourceMode::SourceTextureChannel:
+			break;
+		case ELandscapeTexturePatchAlphaSourceMode::InternalTexture:
+			PatchAlphaUObject = InAlphaSettings.GetAlphaInternalTexture(InOuter);
+			break;
+		case ELandscapeTexturePatchAlphaSourceMode::TextureBackedRenderTarget:
+			PatchAlphaUObject = InAlphaSettings.GetAlphaRenderTarget(InOuter);
+			break;
+		case ELandscapeTexturePatchAlphaSourceMode::TextureAsset:
+			// The alpha texture can be invalid but cannot be a virtual texture
+			if (InAlphaSettings.AlphaTextureAsset && InAlphaSettings.AlphaTextureAsset->VirtualTextureStreaming != 0)
+			{
+				UE_LOG(LogLandscapePatch, Error, TEXT("ULandscapeTexturePatch::GetTexturePatchAlphaResource: Virtual textures are not supported. "));
+				return;
+			}
+			PatchAlphaUObject = InAlphaSettings.AlphaTextureAsset;
+			break;
+		default:
+			checkNoEntry();
+		}
+
+		if (IsValid(PatchAlphaUObject))
+		{
+			OutAlphaTextureResource = PatchAlphaUObject->GetResource();
+		}
+	}
 #endif // WITH_EDITOR
 
 	bool bUseExternalTextureFix = true;
@@ -126,6 +171,137 @@ namespace LandscapeTexturePatchLocals
 		FVector2D TargetPixelSize(Coverage / FVector2D::Max(Resolution - 1, FVector2D(1, 1)));
 		return TargetPixelSize * Resolution;
 	}
+
+	// Alpha source mode implements all default source mode options plus an additional SourceTextureChannel
+	// The SourceTextureChannel mode does not allocate any assets/internal data, so it's commonly equivalent to ELandscapeTexturePatchSourceMode::None
+	inline ELandscapeTexturePatchSourceMode AlphaModeToSourceMode(ELandscapeTexturePatchAlphaSourceMode InMode)
+	{
+		switch (InMode)
+		{
+		// Unlike the other alpha modes, SourceTextureChannel does not add an additional standalone texture
+		case ELandscapeTexturePatchAlphaSourceMode::None:
+		case ELandscapeTexturePatchAlphaSourceMode::SourceTextureChannel:
+			return ELandscapeTexturePatchSourceMode::None;
+		case ELandscapeTexturePatchAlphaSourceMode::InternalTexture:
+			return ELandscapeTexturePatchSourceMode::InternalTexture;
+		case ELandscapeTexturePatchAlphaSourceMode::TextureBackedRenderTarget:
+			return ELandscapeTexturePatchSourceMode::TextureBackedRenderTarget;
+		case ELandscapeTexturePatchAlphaSourceMode::TextureAsset:
+			return ELandscapeTexturePatchSourceMode::TextureAsset;
+		default:
+			checkNoEntry();
+		}
+		return ELandscapeTexturePatchSourceMode::None;
+	}
+}
+
+UTextureRenderTarget2D* FLandscapeTexturePatchAlphaSettings::GetAlphaRenderTarget(UObject* InOuter)
+{
+#if WITH_EDITOR
+	// InOuter is the texture patch or weight patch layer info
+	if (!InOuter || InOuter->IsTemplate())
+	{
+		return nullptr;
+	}
+
+	// Allocate data if needed (see comment in GetHeightRenderTarget)
+	if (AlphaSourceMode == ELandscapeTexturePatchAlphaSourceMode::TextureBackedRenderTarget)
+	{
+		if (!InternalAlphaData || !InternalAlphaData->GetRenderTarget())
+		{
+			SetAlphaSourceMode(AlphaSourceMode, InOuter);
+		}
+
+		return ensure(InternalAlphaData) ? InternalAlphaData->GetRenderTarget() : nullptr;
+	}
+#endif
+
+	return nullptr;
+}
+
+UTexture2D* FLandscapeTexturePatchAlphaSettings::GetAlphaInternalTexture(UObject* InOuter)
+{
+#if WITH_EDITOR
+	// InOuter is the texture patch or weight patch layer info
+	if (!InOuter || InOuter->IsTemplate())
+	{
+		return nullptr;
+	}
+
+	// Allocate data if needed (see comment in GetHeightRenderTarget)
+	if (AlphaSourceMode == ELandscapeTexturePatchAlphaSourceMode::InternalTexture
+		|| AlphaSourceMode == ELandscapeTexturePatchAlphaSourceMode::TextureBackedRenderTarget)
+	{
+		if (!InternalAlphaData || !InternalAlphaData->GetInternalTexture())
+		{
+			SetAlphaSourceMode(AlphaSourceMode, InOuter);
+		}
+
+		return ensure(InternalAlphaData) ? InternalAlphaData->GetInternalTexture() : nullptr;
+	}
+#endif
+
+	return nullptr;
+}
+
+void FLandscapeTexturePatchAlphaSettings::SetAlphaSourceMode(ELandscapeTexturePatchAlphaSourceMode InNewSourceMode, UObject* InOuter, bool bAlwaysMarkDirty/*=true*/)
+{
+#if WITH_EDITOR
+	using namespace LandscapeTexturePatchLocals;
+	if (InNewSourceMode == AlphaSourceMode)
+	{
+		return;
+	}
+
+	// InOuter is the texture patch or weight patch layer info
+	if (!InOuter->IsTemplate())
+	{
+		InOuter->Modify(bAlwaysMarkDirty);
+
+		FVector2D Resolution;
+		// If a weight patch is passed in get the outer texture patch
+		if (ULandscapeTexturePatch* WeightTexturePatch = Cast<ULandscapeTexturePatch>(InOuter->GetOuter()))
+		{
+			Resolution = WeightTexturePatch->GetResolution();
+		}
+		else
+		{
+			// The owner is the patch (used by default weight alpha settings / heightmap alpha settings)
+			ULandscapeTexturePatch* TexturePatch = CastChecked<ULandscapeTexturePatch>(InOuter);
+			Resolution = TexturePatch->GetResolution();
+		}
+
+		// Convert from alpha mode to source mode for transition
+		ELandscapeTexturePatchSourceMode OldMode = AlphaModeToSourceMode(AlphaSourceMode);
+		ELandscapeTexturePatchSourceMode NewMode = AlphaModeToSourceMode(InNewSourceMode);
+
+		TransitionSourceMode<ULandscapeWeightTextureBackedRenderTarget>(OldMode, NewMode, AlphaTextureAsset, InternalAlphaData, AlphaTextureChannel, [&Resolution, this, InOuter]()
+		{
+			check(InOuter);
+			ULandscapeWeightTextureBackedRenderTarget* InternalDataToReturn = NewObject<ULandscapeWeightTextureBackedRenderTarget>(InOuter);
+			InternalDataToReturn->SetFlags(RF_Transactional);
+			InternalDataToReturn->SetResolution(Resolution.X, Resolution.Y);
+			InternalDataToReturn->SetUseAlphaChannel(true);
+			return InternalDataToReturn;
+		});
+
+		// TransitionSourceMode doesn't handle alpha specific mode SourceTextureChannel. 
+		// When the user has no texture channel selected, automatically enable the alpha channel (default color channel in this mode) 
+		if (AlphaTextureChannel == ELandscapeTexturePatchTextureChannel::None && InNewSourceMode == ELandscapeTexturePatchAlphaSourceMode::SourceTextureChannel)
+		{
+			AlphaTextureChannel = ELandscapeTexturePatchTextureChannel::Alpha;
+		}
+	}
+	// In a template, it is not safe to try to allocate a texture, etc. All we do is clear out the
+	// texture asset pointer if it is not needed, to avoid referencing assets unnecessarily.
+	else if (AlphaSourceMode != ELandscapeTexturePatchAlphaSourceMode::TextureAsset)
+	{
+		AlphaTextureAsset = nullptr;
+	}
+
+	AlphaSourceMode = InNewSourceMode;
+	DetailPanelAlphaSourceMode = InNewSourceMode;
+#endif // WITH_EDITOR
 }
 
 #if WITH_EDITOR
@@ -300,16 +476,15 @@ UTextureRenderTarget2D* ULandscapeTexturePatch::ApplyToHeightmap(UE::Landscape::
 		PatchUObject = GetHeightRenderTarget(/*bMarkDirty = */ false);
 		break;
 	case ELandscapeTexturePatchSourceMode::TextureAsset:
-
-		if (IsValid(HeightTextureAsset) && !ensureMsgf(HeightTextureAsset->VirtualTextureStreaming == 0,
-			TEXT("ULandscapeTexturePatch: Virtual textures are not supported")))
+		if (IsValid(HeightTextureAsset) && HeightTextureAsset->VirtualTextureStreaming != 0)
 		{
+			UE_LOG(LogLandscapePatch, Error, TEXT("ULandscapeTexturePatch::ApplyToHeightmap: Invalid texture. Virtual textures are not supported. "));
 			return InCombinedResult;
 		}
 		PatchUObject = HeightTextureAsset;
 		break;
 	default:
-		ensure(false);
+		checkNoEntry();
 	}
 
 	if (!IsValid(PatchUObject))
@@ -323,6 +498,10 @@ UTextureRenderTarget2D* ULandscapeTexturePatch::ApplyToHeightmap(UE::Landscape::
 		return InCombinedResult;
 	}
 
+	// Get the heightmap texture alpha override resource
+	FTextureResource* PatchAlphaTextureResource = nullptr;
+	LandscapeTexturePatchLocals::GetTexturePatchAlphaResource(HeightAlphaSettings, this, PatchAlphaTextureResource);
+
 	// Go ahead and pack everything into a copy of the param struct so we don't have to capture everything
 	// individually in the lambda below.
 	FApplyLandscapeTextureHeightPatchPSParameters ShaderParamsToCopy;
@@ -356,6 +535,7 @@ UTextureRenderTarget2D* ULandscapeTexturePatch::ApplyToHeightmap(UE::Landscape::
 		, OutputResourceName
 		, ShaderParamsToCopy
 		, PatchTextureResource
+		, PatchAlphaTextureResource
 		, DestinationBounds](FRDGBuilder& GraphBuilder)
 		{
 			TRACE_CPUPROFILER_EVENT_SCOPE(LandscapeTextureHeightPatch_Render);
@@ -373,6 +553,19 @@ UTextureRenderTarget2D* ULandscapeTexturePatch::ApplyToHeightmap(UE::Landscape::
 			ShaderParams->InHeightPatch = PatchSRV;
 			ShaderParams->InHeightPatchSampler = TStaticSamplerState<SF_Bilinear, AM_Clamp, AM_Clamp>::GetRHI();
 
+			// Heightmap alpha override resource. Pass white system texture when invalid (no transparency)
+			if (PatchAlphaTextureResource != nullptr && PatchAlphaTextureResource->GetTexture2DRHI())
+			{
+				TRefCountPtr<IPooledRenderTarget> PatchAlphaRenderTarget = CreateRenderTarget(PatchAlphaTextureResource->GetTexture2DRHI(), TEXT("LandscapeAlphaTextureHeightPatch"));
+				FRDGTextureRef PatchAlphaTexture = GraphBuilder.RegisterExternalTexture(PatchAlphaRenderTarget);
+				FRDGTextureSRVRef PatchAlphaSRV = GraphBuilder.CreateSRV(FRDGTextureSRVDesc::CreateForMipLevel(PatchAlphaTexture, 0));
+				ShaderParams->InHeightPatchAlphaTexture = PatchAlphaSRV;
+			}
+			else
+			{
+				ShaderParams->InHeightPatchAlphaTexture = GraphBuilder.CreateSRV(FRDGTextureSRVDesc::Create(GSystemTextures.GetWhiteDummy(GraphBuilder)));
+			}
+
 			ShaderParams->RenderTargets[0] = FRenderTargetBinding(DestinationTexture, ERenderTargetLoadAction::ENoAction, /*InMipIndex = */0);
 
 			FApplyLandscapeTextureHeightPatchPS::AddToRenderGraph(RDG_EVENT_NAME("RenderTextureHeightPatch -> %s", *OutputResourceName), GraphBuilder, ShaderParams, DestinationBounds);
@@ -408,15 +601,15 @@ void ULandscapeTexturePatch::ApplyToWeightmap(UE::Landscape::EditLayers::FRender
 		PatchUObject = GetWeightPatchRenderTarget(PatchInfo);
 		break;
 	case ELandscapeTexturePatchSourceMode::TextureAsset:
-		if (IsValid(PatchInfo->TextureAsset) && !ensureMsgf(PatchInfo->TextureAsset->VirtualTextureStreaming == 0,
-			TEXT("ULandscapeTexturePatch: Virtual textures are not supported")))
+		if (IsValid(PatchInfo->TextureAsset) && PatchInfo->TextureAsset->VirtualTextureStreaming != 0)
 		{
+			UE_LOG(LogLandscapePatch, Error, TEXT("ULandscapeTexturePatch::ApplyToWeightmap: Invalid texture. Virtual textures are not supported. "));
 			return;
 		}
 		PatchUObject = PatchInfo->TextureAsset;
 		break;
 	default:
-		ensure(false);
+		checkNoEntry();
 	}
 
 	if (!IsValid(PatchUObject))
@@ -430,6 +623,11 @@ void ULandscapeTexturePatch::ApplyToWeightmap(UE::Landscape::EditLayers::FRender
 		return;
 	}
 
+	// Get the alpha override or default weight patches alpha settings
+	FTextureResource* PatchAlphaTextureResource = nullptr;
+	FLandscapeTexturePatchAlphaSettings& InAlphaSettings = GetWeightPatchAlphaSettings(PatchInfo);
+	LandscapeTexturePatchLocals::GetTexturePatchAlphaResource(InAlphaSettings, PatchInfo, PatchAlphaTextureResource);
+
 	// Go ahead and pack everything into a copy of the param struct so we don't have to capture everything
 	// individually in the lambda below.
 	FApplyLandscapeTextureWeightPatchPSParameters ShaderParamsToCopy;
@@ -466,6 +664,7 @@ void ULandscapeTexturePatch::ApplyToWeightmap(UE::Landscape::EditLayers::FRender
 		, LandscapeTextureSliceIndex
 		, ShaderParamsToCopy
 		, PatchTextureResource
+		, PatchAlphaTextureResource
 		, DestinationBounds](FRDGBuilder& GraphBuilder)
 		{
 			TRACE_CPUPROFILER_EVENT_SCOPE(LandscapeTextureWeightPatch_Render);
@@ -483,6 +682,19 @@ void ULandscapeTexturePatch::ApplyToWeightmap(UE::Landscape::EditLayers::FRender
 			ShaderParams->InWeightPatch = PatchSRV;
 			ShaderParams->InWeightPatchSampler = TStaticSamplerState<SF_Bilinear, AM_Clamp, AM_Clamp>::GetRHI();
 
+			// Weightmap alpha override resource. Pass white system texture when invalid (no transparency)
+			if (PatchAlphaTextureResource != nullptr && PatchAlphaTextureResource->GetTexture2DRHI())
+			{
+				TRefCountPtr<IPooledRenderTarget> PatchAlphaRenderTarget = CreateRenderTarget(PatchAlphaTextureResource->GetTexture2DRHI(), TEXT("LandscapeAlphaTextureWeightPatch"));
+				FRDGTextureRef PatchAlphaTexture = GraphBuilder.RegisterExternalTexture(PatchAlphaRenderTarget);
+				FRDGTextureSRVRef PatchAlphaSRV = GraphBuilder.CreateSRV(FRDGTextureSRVDesc::CreateForMipLevel(PatchAlphaTexture, 0));
+				ShaderParams->InWeightPatchAlphaTexture = PatchAlphaSRV;
+			}
+			else
+			{
+				ShaderParams->InWeightPatchAlphaTexture = GraphBuilder.CreateSRV(FRDGTextureSRVDesc::Create(GSystemTextures.GetWhiteDummy(GraphBuilder)));
+			}
+			
 			ShaderParams->RenderTargets[0] = FRenderTargetBinding(DestinationTexture, ERenderTargetLoadAction::ENoAction,
 				/*InMipIndex = */0, LandscapeTextureSliceIndex);
 
@@ -602,7 +814,7 @@ void ULandscapeTexturePatch::GetHeightShaderParams(const FTransform& LandscapeHe
 		break;
 	}
 	default:
-		ensure(false);
+		checkNoEntry();
 	}
 
 	ParamsOut.InBlendMode = static_cast<uint32>(BlendMode);
@@ -614,13 +826,15 @@ void ULandscapeTexturePatch::GetHeightShaderParams(const FTransform& LandscapeHe
 	Flags |= (FalloffMode == ELandscapeTexturePatchFalloffMode::RoundedRectangle) ?
 		EShaderFlags::RectangularFalloff : EShaderFlags::None;
 
-	Flags |= bUseTextureAlphaForHeight ?
-		EShaderFlags::ApplyPatchAlpha : EShaderFlags::None;
-
 	Flags |= bNativeEncoding ?
 		EShaderFlags::InputIsPackedHeight : EShaderFlags::None;
 
 	ParamsOut.InFlags = static_cast<uint8>(Flags);
+	check(!bNativeEncoding || HeightTextureChannel == ELandscapeTexturePatchTextureChannel::None);
+	ParamsOut.InTextureChannel = static_cast<uint8>(HeightTextureChannel);
+
+	ParamsOut.InAlphaMode = static_cast<uint8>(HeightAlphaSettings.GetAlphaSourceMode());
+	ParamsOut.InAlphaTextureChannel = static_cast<uint8>(HeightAlphaSettings.AlphaTextureChannel);
 }
 
 void ULandscapeTexturePatch::GetWeightShaderParams(
@@ -637,7 +851,7 @@ void ULandscapeTexturePatch::GetWeightShaderParams(
 		DestinationBoundsOut, ParamsOut.InEdgeUVDeadBorder, ParamsOut.InFalloffWorldMargin);
 
 	// Use the override blend mode if present, otherwise fall back to more general blend mode.
-	ParamsOut.InBlendMode = static_cast<uint32>(WeightPatchInfo->bOverrideBlendMode ? WeightPatchInfo->OverrideBlendMode : BlendMode);
+	ParamsOut.InBlendMode = static_cast<uint32>(WeightPatchInfo->bOverrideBlendMode ? WeightPatchInfo->OverrideBlendMode : WeightPatchesBlendMode);
 
 	// Pack our booleans into a bitfield
 	using EShaderFlags = FApplyLandscapeTextureHeightPatchPS::EFlags;
@@ -646,10 +860,11 @@ void ULandscapeTexturePatch::GetWeightShaderParams(
 	Flags |= (FalloffMode == ELandscapeTexturePatchFalloffMode::RoundedRectangle) ?
 		EShaderFlags::RectangularFalloff : EShaderFlags::None;
 
-	Flags |= WeightPatchInfo->bUseAlphaChannel ?
-		EShaderFlags::ApplyPatchAlpha : EShaderFlags::None;
-
 	ParamsOut.InFlags = static_cast<uint8>(Flags);
+	ParamsOut.InTextureChannel = static_cast<uint8>(WeightPatchInfo->WeightPatchTextureChannel);
+
+	ParamsOut.InAlphaMode = static_cast<uint8>(GetWeightPatchAlphaSourceMode(WeightPatchInfo->WeightmapLayerName));
+	ParamsOut.InAlphaTextureChannel = static_cast<uint8>(GetWeightPatchAlphaTextureChannel(WeightPatchInfo->WeightmapLayerName));
 }
 
 // This function determines how our internal height render targets get converted to the format that gets
@@ -783,7 +998,7 @@ void ULandscapeTexturePatch::RequestReinitializeWeights()
 
 			if (FoundPatches.IsEmpty())
 			{
-				AddWeightPatch(WeightmapLayerName, ELandscapeTexturePatchSourceMode::InternalTexture, false);
+				CreateWeightPatch(WeightmapLayerName, ELandscapeTexturePatchSourceMode::InternalTexture, ELandscapeTexturePatchAlphaSourceMode::None);
 				WeightPatches.Last()->bReinitializeOnNextRender = true;
 				WeightPatches.Last()->bEditVisibilityLayer = bIsVisibilityLayer;
 			}
@@ -836,7 +1051,7 @@ void ULandscapeTexturePatch::ReinitializeHeight(UTextureRenderTarget2D* InCombin
 		return;
 	}
 
-	SetUseAlphaChannelForHeight(false);
+	SetHeightAlphaSourceMode(ELandscapeTexturePatchAlphaSourceMode::None);
 	SetBlendMode(ELandscapeTexturePatchBlendMode::AlphaBlend);
 	ResetHeightRenderTargetFormat();
 
@@ -979,6 +1194,9 @@ void ULandscapeTexturePatch::ReinitializeWeightPatch(ULandscapeWeightPatchTextur
 			PatchInfo->InternalData->GetRenderTarget()->Modify();
 		}
 	}
+
+	// Clear the alpha overrides
+	PatchInfo->SetWeightPatchAlphaSourceMode(ELandscapeTexturePatchAlphaSourceMode::None, /*bInOverrideAlpha =*/false);
 	
 	if (!ensure(PatchInfo->InternalData))
 	{
@@ -986,7 +1204,7 @@ void ULandscapeTexturePatch::ReinitializeWeightPatch(ULandscapeWeightPatchTextur
 	}
 
 	PatchInfo->InternalData->SetUseAlphaChannel(false);
-	if (BlendMode != ELandscapeTexturePatchBlendMode::AlphaBlend)
+	if (WeightPatchesBlendMode != ELandscapeTexturePatchBlendMode::AlphaBlend)
 	{
 		PatchInfo->bOverrideBlendMode = true;
 		PatchInfo->OverrideBlendMode = ELandscapeTexturePatchBlendMode::AlphaBlend;
@@ -1109,6 +1327,7 @@ bool ULandscapeTexturePatch::WeightPatchCanRender(const ULandscapeWeightPatchTex
 {
 	return (InWeightPatch.SourceMode != ELandscapeTexturePatchSourceMode::None 
 		// If source mode is texture asset, we need to have an asset to read from
+		// It's valid for the weight patch alpha texture to be invalid, a default system texture will be supplied for the alpha
 		&& (InWeightPatch.SourceMode != ELandscapeTexturePatchSourceMode::TextureAsset || InWeightPatch.TextureAsset))
 		// If reinitializing, we need to read from the render call
 		|| InWeightPatch.bReinitializeOnNextRender;
@@ -1129,7 +1348,32 @@ void ULandscapeTexturePatch::GetRenderDependencies(TSet<UObject*>& OutDependenci
 		OutDependencies.Add(HeightTextureAsset);
 	}
 
-	Algo::ForEach(WeightPatches, [&OutDependencies](const TObjectPtr<ULandscapeWeightPatchTextureInfo>& InWeightPatch)
+	// Heightmap alpha texture
+	if (HeightAlphaSettings.GetAlphaSourceMode() == ELandscapeTexturePatchAlphaSourceMode::TextureAsset)
+	{
+		OutDependencies.Add(HeightAlphaSettings.AlphaTextureAsset);
+	}
+	else if (HeightAlphaSettings.GetAlphaSourceMode() == ELandscapeTexturePatchAlphaSourceMode::InternalTexture
+		&& HeightAlphaSettings.InternalAlphaData && HeightAlphaSettings.InternalAlphaData->GetInternalTexture())
+	{
+		OutDependencies.Add(HeightAlphaSettings.InternalAlphaData->GetInternalTexture());
+	}
+
+	// Weightmap alpha texture (Default alpha for all weight patches) 
+	if (WeightPatchesAlphaSettings.GetAlphaSourceMode() == ELandscapeTexturePatchAlphaSourceMode::TextureAsset)
+	{
+		OutDependencies.Add(WeightPatchesAlphaSettings.AlphaTextureAsset);
+	}
+	else if (WeightPatchesAlphaSettings.GetAlphaSourceMode() == ELandscapeTexturePatchAlphaSourceMode::InternalTexture)
+	{
+		if (WeightPatchesAlphaSettings.InternalAlphaData && WeightPatchesAlphaSettings.InternalAlphaData->GetInternalTexture())
+		{
+			OutDependencies.Add(WeightPatchesAlphaSettings.InternalAlphaData->GetInternalTexture());
+		}
+	}
+
+	// Individual weight patch data
+	Algo::ForEach(WeightPatches, [&OutDependencies, this](const TObjectPtr<ULandscapeWeightPatchTextureInfo>& InWeightPatch)
 		{ 
 			if (IsValid(InWeightPatch))
 			{
@@ -1143,6 +1387,23 @@ void ULandscapeTexturePatch::GetRenderDependencies(TSet<UObject*>& OutDependenci
 				{
 					OutDependencies.Add(InWeightPatch->TextureAsset);
 				}
+
+				// If this weight patch overrides the default alpha, add the additional dependencies
+				if (InWeightPatch->bOverrideWeightAlphaSettings)
+				{
+					const FLandscapeTexturePatchAlphaSettings& PatchAlphaSettings = InWeightPatch->WeightPatchAlphaSettings;
+					if (PatchAlphaSettings.GetAlphaSourceMode() == ELandscapeTexturePatchAlphaSourceMode::TextureAsset)
+					{
+						OutDependencies.Add(PatchAlphaSettings.AlphaTextureAsset);
+					}
+					else if (PatchAlphaSettings.GetAlphaSourceMode() == ELandscapeTexturePatchAlphaSourceMode::InternalTexture)
+					{
+						if (PatchAlphaSettings.InternalAlphaData && PatchAlphaSettings.InternalAlphaData->GetInternalTexture())
+						{
+							OutDependencies.Add(PatchAlphaSettings.InternalAlphaData->GetInternalTexture());
+						}
+					}
+				}
 			}
 		});
 }
@@ -1184,12 +1445,25 @@ TStructOnScope<FActorComponentInstanceData> ULandscapeTexturePatch::GetComponent
 	auto SetAllInternalDataTransactionalFlags = [this, &SetObjectTransactionalFlag, &SetInternalDataTransactionalFlags](bool bOn)
 	{
 		SetInternalDataTransactionalFlags(HeightInternalData, bOn);
+
+		// Default heightmap alpha
+		SetInternalDataTransactionalFlags(HeightAlphaSettings.InternalAlphaData, bOn);
+
+		// Default weightmap alpha
+		SetInternalDataTransactionalFlags(WeightPatchesAlphaSettings.InternalAlphaData, bOn);
+
 		for (const TObjectPtr<ULandscapeWeightPatchTextureInfo>& WeightPatch : WeightPatches)
 		{
 			if (IsValid(WeightPatch))
 			{
 				SetObjectTransactionalFlag(WeightPatch, bOn);
 				SetInternalDataTransactionalFlags(WeightPatch->InternalData, bOn);
+				
+				// Override weightmap alpha
+				if (WeightPatch->bOverrideWeightAlphaSettings)
+				{
+					SetInternalDataTransactionalFlags(WeightPatch->WeightPatchAlphaSettings.InternalAlphaData, bOn);
+				}
 			}
 		}
 	};
@@ -1371,11 +1645,22 @@ void ULandscapeTexturePatch::SetResolution(FVector2D ResolutionIn)
 	};
 
 	ResizePatch(HeightSourceMode, HeightInternalData);
+	// Height alpha data
+	ResizePatch(LandscapeTexturePatchLocals::AlphaModeToSourceMode(HeightAlphaSettings.GetAlphaSourceMode()), HeightAlphaSettings.InternalAlphaData);
+	// Default weight alpha data 
+	ResizePatch(LandscapeTexturePatchLocals::AlphaModeToSourceMode(WeightPatchesAlphaSettings.GetAlphaSourceMode()), WeightPatchesAlphaSettings.InternalAlphaData);
+
 	for (const TObjectPtr<ULandscapeWeightPatchTextureInfo>& WeightPatch : WeightPatches)
 	{
 		if (IsValid(WeightPatch))
 		{
 			ResizePatch(WeightPatch->SourceMode, WeightPatch->InternalData);
+
+			// Override weight alpha data
+			if (WeightPatch->bOverrideWeightAlphaSettings)
+			{
+				ResizePatch(LandscapeTexturePatchLocals::AlphaModeToSourceMode(WeightPatch->WeightPatchAlphaSettings.GetAlphaSourceMode()), WeightPatch->WeightPatchAlphaSettings.InternalAlphaData);
+			}
 		}
 	}
 }
@@ -1434,6 +1719,47 @@ bool ULandscapeTexturePatch::GetInitResolutionFromLandscape(float ResolutionMult
 }
 
 #if WITH_EDITOR
+void ULandscapeTexturePatch::PostLoad()
+{
+	Super::PostLoad();
+
+	// Handle weight/height patch alpha channel and blend mode deprecation
+	if (GetLinkerCustomVersion(FFortniteMainBranchObjectVersion::GUID) < FFortniteMainBranchObjectVersion::LandscapeTexturePatchAlphaOverride)
+	{
+		// Weight patches used the height blend mode as the default fallback. Copy the legacy height fallback to the weight patches default blend mode
+		WeightPatchesBlendMode = BlendMode;
+	
+		// Reset the height texture channel if the user is using native encoding/internal texture
+		if (HeightSourceMode == ELandscapeTexturePatchSourceMode::InternalTexture || HeightEncoding == ELandscapeTextureHeightPatchEncoding::NativePackedHeight)
+		{
+			HeightTextureChannel = ELandscapeTexturePatchTextureChannel::None;
+		}
+
+PRAGMA_DISABLE_DEPRECATION_WARNINGS
+		// Convert heightmap bUseAlphaChannel legacy flag to the new alpha enum
+		if (bUseTextureAlphaForHeight_DEPRECATED)
+		{
+			SetHeightAlphaSourceMode(ELandscapeTexturePatchAlphaSourceMode::SourceTextureChannel, /*bAlwaysMarkDirty =*/false);
+		}
+
+PRAGMA_ENABLE_DEPRECATION_WARNINGS
+		
+		// Convert weightmap bUseAlphaChannel legacy flags to the new alpha enum
+		for (TObjectPtr<ULandscapeWeightPatchTextureInfo> WeightPatch : WeightPatches)
+		{
+PRAGMA_DISABLE_DEPRECATION_WARNINGS
+			if (WeightPatch && WeightPatch->bUseAlphaChannel_DEPRECATED)
+			{
+				// Set the override alpha source mode
+				WeightPatch->SetWeightPatchAlphaSourceMode(ELandscapeTexturePatchAlphaSourceMode::SourceTextureChannel,
+					/*bInOverrideAlpha =*/true,
+					/*bAlwaysMarkDirty = */ false);
+PRAGMA_ENABLE_DEPRECATION_WARNINGS
+			}
+		}
+	}
+}
+
 void ULandscapeTexturePatch::PostEditChangeProperty(FPropertyChangedEvent& PropertyChangedEvent)
 {
 	using namespace LandscapeTexturePatchLocals;
@@ -1473,6 +1799,21 @@ void ULandscapeTexturePatch::PostEditChangeProperty(FPropertyChangedEvent& Prope
 				}
 			}
 		}
+		else if (PropertyChangedEvent.MemberProperty->GetFName() == GET_MEMBER_NAME_CHECKED(ULandscapeTexturePatch, WeightPatchesAlphaSettings))
+		{
+			// Changing the default weightmap alpha source mode
+			if (PropertyChangedEvent.Property->GetFName() == GET_MEMBER_NAME_CHECKED(FLandscapeTexturePatchAlphaSettings, DetailPanelAlphaSourceMode))
+			{
+				SetWeightPatchAlphaSourceMode(FName(), WeightPatchesAlphaSettings.DetailPanelAlphaSourceMode);
+			}
+		}
+		else if (PropertyChangedEvent.MemberProperty->GetFName() == GET_MEMBER_NAME_CHECKED(ULandscapeTexturePatch, HeightAlphaSettings))
+		{
+			if (PropertyChangedEvent.Property->GetFName() == GET_MEMBER_NAME_CHECKED(FLandscapeTexturePatchAlphaSettings, DetailPanelAlphaSourceMode))
+			{
+				SetHeightAlphaSourceMode(HeightAlphaSettings.DetailPanelAlphaSourceMode);
+			}
+		}
 		else if (PropertyChangedEvent.Property->GetFName() == GET_MEMBER_NAME_CHECKED(FLandscapeTexturePatchEncodingSettings, ZeroInEncoding)
 			|| PropertyChangedEvent.Property->GetFName() == GET_MEMBER_NAME_CHECKED(FLandscapeTexturePatchEncodingSettings, WorldSpaceEncodingScale))
 		{
@@ -1499,6 +1840,35 @@ void ULandscapeWeightPatchTextureInfo::PostEditChangeProperty(FPropertyChangedEv
 	Super::PostEditChangeProperty(PropertyChangedEvent);
 }
 
+void ULandscapeWeightPatchTextureInfo::PostEditChangeChainProperty(FPropertyChangedChainEvent& PropertyChangedEvent)
+{
+	Super::PostEditChangeChainProperty(PropertyChangedEvent);
+
+	FProperty* Property = PropertyChangedEvent.Property;
+	FProperty* MemberProperty = PropertyChangedEvent.PropertyChain.GetActiveMemberNode() ? PropertyChangedEvent.PropertyChain.GetActiveMemberNode()->GetValue() : nullptr;
+
+	if (Property && MemberProperty)
+	{
+		// Changing the alpha settings for a specific weight patch
+		if (MemberProperty->GetFName() == GET_MEMBER_NAME_CHECKED(ULandscapeWeightPatchTextureInfo, WeightPatchAlphaSettings))
+		{
+			if (Property->GetFName() == GET_MEMBER_NAME_CHECKED(FLandscapeTexturePatchAlphaSettings, DetailPanelAlphaSourceMode))
+			{
+				SetWeightPatchAlphaSourceMode(WeightPatchAlphaSettings.DetailPanelAlphaSourceMode, bOverrideWeightAlphaSettings);
+			}
+			else
+			{
+				// When bOverrideWeightAlphaSettings changes the property does not get passed to this function correctly
+				// As a fallback, clear the internal data whenever the weight patch no longer overrides the default
+				if (!bOverrideWeightAlphaSettings)
+				{
+					SetWeightPatchAlphaSourceMode(ELandscapeTexturePatchAlphaSourceMode::None, /*bInOverrideAlpha =*/false);
+				}
+			}
+		}
+	}
+}
+
 void ULandscapeWeightPatchTextureInfo::PreDuplicate(FObjectDuplicationParameters& DupParams)
 {
 	// TODO: It seems like this whole overload shouldn't be necessary, because we should get PreDuplicate calls
@@ -1511,6 +1881,15 @@ void ULandscapeWeightPatchTextureInfo::PreDuplicate(FObjectDuplicationParameters
 	{
 		InternalData->CopyToInternalTexture();
 	}
+
+	// Handle weightmap alpha override
+	if (bOverrideWeightAlphaSettings)
+	{
+		if (WeightPatchAlphaSettings.GetAlphaSourceMode() == ELandscapeTexturePatchAlphaSourceMode::TextureBackedRenderTarget && WeightPatchAlphaSettings.InternalAlphaData)
+		{
+			WeightPatchAlphaSettings.InternalAlphaData->CopyToInternalTexture();
+		}
+	}
 }
 #endif // WITH_EDITOR
 
@@ -1546,6 +1925,62 @@ void ULandscapeWeightPatchTextureInfo::SetSourceMode(ELandscapeTexturePatchSourc
 #endif // WITH_EDITOR
 }
 
+void ULandscapeWeightPatchTextureInfo::SetWeightPatchAlphaSourceMode(ELandscapeTexturePatchAlphaSourceMode InSourceMode, bool bInOverrideAlpha, bool bAlwaysMarkDirty/*=true*/)
+{
+#if WITH_EDITOR
+	// Early return if override and source do not change
+	if (bOverrideWeightAlphaSettings == bInOverrideAlpha && WeightPatchAlphaSettings.GetAlphaSourceMode() == InSourceMode)
+	{
+		return;
+	}
+
+	Modify(bAlwaysMarkDirty);
+	bOverrideWeightAlphaSettings = bInOverrideAlpha;
+
+	// When bInOverrideAlpha is false, the weightmap layer will fallback to the shared weightmap alpha. Reset the override source mode to clear internal data
+	if (bOverrideWeightAlphaSettings)
+	{
+		WeightPatchAlphaSettings.SetAlphaSourceMode(InSourceMode, this);
+	}
+	else
+	{
+		WeightPatchAlphaSettings.SetAlphaSourceMode(ELandscapeTexturePatchAlphaSourceMode::None, this);
+	}
+
+	// The primary internal source texture tracks the alpha status
+	if (InternalData)
+	{
+		InternalData->SetUseAlphaChannel(InSourceMode != ELandscapeTexturePatchAlphaSourceMode::None);
+	}
+#endif // WITH_EDITOR
+}
+
+void ULandscapeWeightPatchTextureInfo::SetWeightPatchTextureChannel(ELandscapeTexturePatchTextureChannel InTextureChannel, bool bAlwaysMarkDirty)
+{
+#if WITH_EDITOR
+	if (WeightPatchTextureChannel == InTextureChannel)
+	{
+		return;
+	}
+	Modify(bAlwaysMarkDirty);
+
+	WeightPatchTextureChannel = InTextureChannel;
+#endif // WITH_EDITOR
+}
+
+void ULandscapeWeightPatchTextureInfo::SetWeightPatchAlphaTextureChannel(ELandscapeTexturePatchTextureChannel InTextureChannel, bool bAlwaysMarkDirty)
+{
+#if WITH_EDITOR
+	if (WeightPatchAlphaSettings.AlphaTextureChannel == InTextureChannel)
+	{
+		return;
+	}
+	Modify(bAlwaysMarkDirty);
+
+	WeightPatchAlphaSettings.AlphaTextureChannel = InTextureChannel;
+#endif // WITH_EDITOR
+}
+
 #if WITH_EDITOR
 void ULandscapeWeightPatchTextureInfo::TransitionSourceModeInternal(ELandscapeTexturePatchSourceMode OldMode, ELandscapeTexturePatchSourceMode NewMode)
 {
@@ -1557,7 +1992,7 @@ void ULandscapeWeightPatchTextureInfo::TransitionSourceModeInternal(ELandscapeTe
 		Resolution = OwningPatch->GetResolution();
 	}
 
-	TransitionSourceMode<ULandscapeWeightTextureBackedRenderTarget>(SourceMode, NewMode, TextureAsset, InternalData, [&Resolution, this]()
+	TransitionSourceMode<ULandscapeWeightTextureBackedRenderTarget>(SourceMode, NewMode, TextureAsset, InternalData, WeightPatchTextureChannel, [&Resolution, this]()
 	{
 		ULandscapeWeightTextureBackedRenderTarget* InternalDataToReturn = NewObject<ULandscapeWeightTextureBackedRenderTarget>(this);
 		InternalDataToReturn->SetFlags(RF_Transactional);
@@ -1594,12 +2029,53 @@ void ULandscapeTexturePatch::SetHeightSourceMode(ELandscapeTexturePatchSourceMod
 #endif // WITH_EDITOR
 }
 
+void ULandscapeTexturePatch::SetHeightTextureChannel(ELandscapeTexturePatchTextureChannel InTextureChannel, bool bAlwaysMarkDirty)
+{
+#if WITH_EDITOR
+	if (HeightTextureChannel == InTextureChannel)
+	{
+		return;
+	}
+
+	// When the height uses internal texture / native packed height we limit the user to RG texture color channels
+	if (HeightSourceMode == ELandscapeTexturePatchSourceMode::InternalTexture || HeightEncoding == ELandscapeTextureHeightPatchEncoding::NativePackedHeight)
+	{
+		Modify(bAlwaysMarkDirty);
+		HeightTextureChannel = ELandscapeTexturePatchTextureChannel::None;
+		return;
+	}
+
+	Modify(bAlwaysMarkDirty);
+	HeightTextureChannel = InTextureChannel;
+#endif // WITH_EDITOR
+}
+
+void ULandscapeTexturePatch::SetHeightAlphaSourceMode(ELandscapeTexturePatchAlphaSourceMode InNewMode, bool bAlwaysMarkDirty/*=true*/)
+{
+#if WITH_EDITOR
+	HeightAlphaSettings.SetAlphaSourceMode(InNewMode, this, bAlwaysMarkDirty);
+#endif // WITH_EDITOR
+}
+
+void ULandscapeTexturePatch::SetHeightAlphaTextureChannel(ELandscapeTexturePatchTextureChannel InTextureChannel, bool bAlwaysMarkDirty)
+{
+#if WITH_EDITOR
+	if (HeightAlphaSettings.AlphaTextureChannel == InTextureChannel)
+	{
+		return;
+	}
+	Modify(bAlwaysMarkDirty);
+
+	HeightAlphaSettings.AlphaTextureChannel = InTextureChannel;
+#endif // WITH_EDITOR
+}
+
 #if WITH_EDITOR
 void ULandscapeTexturePatch::TransitionHeightSourceModeInternal(ELandscapeTexturePatchSourceMode OldMode, ELandscapeTexturePatchSourceMode NewMode)
 {
 	using namespace LandscapeTexturePatchLocals;
 
-	TransitionSourceMode<ULandscapeHeightTextureBackedRenderTarget>(HeightSourceMode, NewMode, HeightTextureAsset, HeightInternalData, [this]()
+	TransitionSourceMode<ULandscapeHeightTextureBackedRenderTarget>(HeightSourceMode, NewMode, HeightTextureAsset, HeightInternalData, HeightTextureChannel, [this]()
 	{
 		ULandscapeHeightTextureBackedRenderTarget* InternalDataToReturn = NewObject<ULandscapeHeightTextureBackedRenderTarget>(this);
 		InternalDataToReturn->SetFlags(RF_Transactional);
@@ -1609,13 +2085,22 @@ void ULandscapeTexturePatch::TransitionHeightSourceModeInternal(ELandscapeTextur
 
 		return InternalDataToReturn;
 	});
+
+	// When the Height Patch data uses InternalTexture mode/Native encoding we explicitly use the RG texture color channels
+	if (NewMode == ELandscapeTexturePatchSourceMode::InternalTexture || HeightEncoding == ELandscapeTextureHeightPatchEncoding::NativePackedHeight)
+	{
+		HeightTextureChannel = ELandscapeTexturePatchTextureChannel::None;
+	}
 }
 #endif
 
 void ULandscapeTexturePatch::SetHeightTextureAsset(UTexture* TextureIn)
 {
-	ensureMsgf(!TextureIn || TextureIn->VirtualTextureStreaming == 0,
-		TEXT("ULandscapeTexturePatch::SetHeightTextureAsset: Virtual textures are not supported."));
+	if (TextureIn && TextureIn->VirtualTextureStreaming != 0)
+	{
+		UE_LOG(LogLandscapePatch, Error, TEXT("ULandscapeTexturePatch::SetHeightTextureAsset: Invalid texture. Virtual textures are not supported."));
+		return;
+	}
 	HeightTextureAsset = TextureIn;
 }
 
@@ -1681,6 +2166,24 @@ UTexture2D* ULandscapeTexturePatch::GetHeightInternalTexture()
 	return nullptr;
 }
 
+UTextureRenderTarget2D* ULandscapeTexturePatch::GetHeightAlphaRenderTarget()
+{
+#if WITH_EDITOR
+	return HeightAlphaSettings.GetAlphaRenderTarget(this);
+#else
+	return nullptr;
+#endif // WITH_EDITOR
+}
+
+UTexture2D* ULandscapeTexturePatch::GetHeightAlphaInternalTexture()
+{
+#if WITH_EDITOR
+	return HeightAlphaSettings.GetAlphaInternalTexture(this);
+#else
+	return nullptr;
+#endif // WITH_EDITOR
+}
+
 void ULandscapeTexturePatch::UpdateHeightConvertToNativeParamsIfNeeded()
 {
 #if WITH_EDITOR
@@ -1731,7 +2234,13 @@ void ULandscapeTexturePatch::ResetHeightEncodingMode(ELandscapeTextureHeightPatc
 void ULandscapeTexturePatch::ResetHeightRenderTargetFormat()
 {
 	SetHeightRenderTargetFormat(HeightEncoding == ELandscapeTextureHeightPatchEncoding::NativePackedHeight ? ETextureRenderTargetFormat::RTF_RGBA8
-		: bUseTextureAlphaForHeight ? ETextureRenderTargetFormat::RTF_RGBA32f : ETextureRenderTargetFormat::RTF_R32f);
+		: (HeightAlphaSettings.GetAlphaSourceMode() != ELandscapeTexturePatchAlphaSourceMode::None) ? ETextureRenderTargetFormat::RTF_RGBA32f : ETextureRenderTargetFormat::RTF_R32f);
+
+	// When the height uses internal texture / native packed height we limit the user to RG texture color channels
+	if (HeightEncoding == ELandscapeTextureHeightPatchEncoding::NativePackedHeight || HeightSourceMode == ELandscapeTexturePatchSourceMode::InternalTexture)
+	{
+		HeightTextureChannel = ELandscapeTexturePatchTextureChannel::None;
+	}
 }
 #endif
 
@@ -1758,8 +2267,13 @@ void ULandscapeTexturePatch::SetHeightRenderTargetFormat(ETextureRenderTargetFor
 	}
 }
 
-
+// Deprecated
 void ULandscapeTexturePatch::AddWeightPatch(const FName& WeightmapLayerName, ELandscapeTexturePatchSourceMode SourceMode, bool bUseAlphaChannel)
+{
+	CreateWeightPatch(WeightmapLayerName, SourceMode, bUseAlphaChannel ? ELandscapeTexturePatchAlphaSourceMode::SourceTextureChannel : ELandscapeTexturePatchAlphaSourceMode::None);
+}
+
+void ULandscapeTexturePatch::CreateWeightPatch(const FName& InWeightmapLayerName, ELandscapeTexturePatchSourceMode InSourceMode, ELandscapeTexturePatchAlphaSourceMode InAlphaSourceMode)
 {
 #if WITH_EDITOR
 	using namespace LandscapeTexturePatchLocals;
@@ -1772,15 +2286,17 @@ void ULandscapeTexturePatch::AddWeightPatch(const FName& WeightmapLayerName, ELa
 			continue;
 		}
 
-		if (WeightPatch->WeightmapLayerName == WeightmapLayerName)
+		if (WeightPatch->WeightmapLayerName == InWeightmapLayerName)
 		{
-			if (WeightPatch->SourceMode != SourceMode)
+			if (WeightPatch->SourceMode != InSourceMode)
 			{
-				WeightPatch->SetSourceMode(SourceMode);
+				WeightPatch->SetSourceMode(InSourceMode);
 			}
-			if (IsValid(WeightPatch->InternalData))
+			// Change the alpha override if needed. This will call UseAlphaChannel on the source internal data
+			const FLandscapeTexturePatchAlphaSettings& AlphaSettings = WeightPatch->WeightPatchAlphaSettings;
+			if (AlphaSettings.GetAlphaSourceMode() != InAlphaSourceMode)
 			{
-				WeightPatch->InternalData->SetUseAlphaChannel(bUseAlphaChannel);
+				WeightPatch->SetWeightPatchAlphaSourceMode(InAlphaSourceMode, /*bInOverrideAlpha =*/true);
 			}
 			return;
 		}
@@ -1795,10 +2311,14 @@ void ULandscapeTexturePatch::AddWeightPatch(const FName& WeightmapLayerName, ELa
 	}
 	ULandscapeWeightPatchTextureInfo* NewWeightPatch = NewObject<ULandscapeWeightPatchTextureInfo>(this, NAME_None, NewObjectFlags);
 
-	NewWeightPatch->WeightmapLayerName = WeightmapLayerName;
-	NewWeightPatch->SourceMode = SourceMode;
-	NewWeightPatch->DetailPanelSourceMode = SourceMode;
-	NewWeightPatch->bUseAlphaChannel = bUseAlphaChannel;
+	NewWeightPatch->WeightmapLayerName = InWeightmapLayerName;
+	NewWeightPatch->SourceMode = InSourceMode;
+	NewWeightPatch->DetailPanelSourceMode = InSourceMode;
+
+	if (InAlphaSourceMode != ELandscapeTexturePatchAlphaSourceMode::None)
+	{
+		NewWeightPatch->SetWeightPatchAlphaSourceMode(InAlphaSourceMode, /*bInOverrideAlpha =*/true);
+	}
 
 	if (NewWeightPatch->SourceMode == ELandscapeTexturePatchSourceMode::InternalTexture
 		|| NewWeightPatch->SourceMode == ELandscapeTexturePatchSourceMode::TextureBackedRenderTarget)
@@ -1806,10 +2326,9 @@ void ULandscapeTexturePatch::AddWeightPatch(const FName& WeightmapLayerName, ELa
 		NewWeightPatch->InternalData = NewObject<ULandscapeWeightTextureBackedRenderTarget>(NewWeightPatch);
 		NewWeightPatch->InternalData->SetFlags(RF_Transactional);
 		NewWeightPatch->InternalData->SetResolution(ResolutionX, ResolutionY);
-		NewWeightPatch->InternalData->SetUseAlphaChannel(bUseAlphaChannel);
+		NewWeightPatch->InternalData->SetUseAlphaChannel(InAlphaSourceMode != ELandscapeTexturePatchAlphaSourceMode::None);
 		NewWeightPatch->InternalData->Initialize();
 	}
-
 	WeightPatches.Add(NewWeightPatch);
 #endif // WITH_EDITOR
 }
@@ -1852,22 +2371,167 @@ TArray<FName> ULandscapeTexturePatch::GetAllWeightPatchLayerNames()
 	return Names;
 }
 
+ELandscapeTexturePatchAlphaSourceMode ULandscapeTexturePatch::GetWeightPatchAlphaSourceMode(const FName& InWeightmapLayerName) const
+{
+	// Get the default weightmap alpha source mode for all weightmap patches
+	if (InWeightmapLayerName.IsNone())
+	{
+		return WeightPatchesAlphaSettings.GetAlphaSourceMode();
+	}
+
+	for (const TObjectPtr<ULandscapeWeightPatchTextureInfo>& WeightPatch : WeightPatches)
+	{
+		if (IsValid(WeightPatch) && WeightPatch->WeightmapLayerName == InWeightmapLayerName)
+		{
+			// Returns the override source mode if and only if the alpha settings are overridden
+			return GetWeightPatchAlphaSettingsConst(WeightPatch).GetAlphaSourceMode();
+		}
+	}
+
+	UE_LOG(LogLandscapePatch, Warning, TEXT("ULandscapeTexturePatch::GetWeightPatchAlphaSourceMode: Unable to get data for weight layer %s"), *InWeightmapLayerName.ToString());
+	return WeightPatchesAlphaSettings.GetAlphaSourceMode();
+}
+
+// Deprecated
 void ULandscapeTexturePatch::SetUseAlphaChannelForWeightPatch(const FName& InWeightmapLayerName, bool bUseAlphaChannel)
 {
+	SetWeightPatchAlphaSourceMode(InWeightmapLayerName, bUseAlphaChannel ? ELandscapeTexturePatchAlphaSourceMode::SourceTextureChannel : ELandscapeTexturePatchAlphaSourceMode::None);
+}
+
+void ULandscapeTexturePatch::SetWeightPatchAlphaSourceMode(const FName& InWeightmapLayerName, ELandscapeTexturePatchAlphaSourceMode InNewSourceMode, bool bInOverrideAlpha, bool bAlwaysMarkDirty)
+{
+	// Set the default weightmap alpha source mode for all weightmap patches
+	if (InWeightmapLayerName.IsNone())
+	{
+		WeightPatchesAlphaSettings.SetAlphaSourceMode(InNewSourceMode, this);
+		return;
+	}
+
+	// Set the override weightmap source mode for a specific weight patch
+	for (const TObjectPtr<ULandscapeWeightPatchTextureInfo>& WeightPatch : WeightPatches)
+	{
+		if (IsValid(WeightPatch) && WeightPatch->WeightmapLayerName == InWeightmapLayerName)
+		{
+			WeightPatch->SetWeightPatchAlphaSourceMode(InNewSourceMode, bInOverrideAlpha, bAlwaysMarkDirty);
+			return;
+		}
+	}
+	const FString LayerNameString = InWeightmapLayerName.ToString();
+	UE_LOG(LogLandscapePatch, Warning, TEXT("ULandscapeTexturePatch::SetWeightPatchAlphaSourceMode: Unable to find data for weight layer %s"), *LayerNameString);
+}
+
+const UTexture* ULandscapeTexturePatch::GetWeightPatchAlphaTextureAsset(const FName& InWeightmapLayerName) const
+{	
+	// Get the default weightmap alpha texture asset for all weightmap patches
+	if (InWeightmapLayerName.IsNone())
+	{
+		return WeightPatchesAlphaSettings.AlphaTextureAsset;
+	}
+
+	for (const TObjectPtr<ULandscapeWeightPatchTextureInfo>& WeightPatch : WeightPatches)
+	{
+		if (IsValid(WeightPatch) && WeightPatch->WeightmapLayerName == InWeightmapLayerName)
+		{
+			// Returns the override value if and only if the alpha settings are overridden
+			return  GetWeightPatchAlphaSettingsConst(WeightPatch).AlphaTextureAsset;
+		}
+	}
+
+	UE_LOG(LogLandscapePatch, Warning, TEXT("ULandscapeTexturePatch::GetWeightPatchAlphaTextureAsset: Unable to get data for weight layer %s"), *InWeightmapLayerName.ToString());
+	return WeightPatchesAlphaSettings.AlphaTextureAsset;
+}
+
+void ULandscapeTexturePatch::SetWeightPatchAlphaTextureAsset(const FName& InWeightmapLayerName, UTexture* InTexture, bool bInOverrideAlpha, bool bAlwaysMarkDirty)
+{
+	if (InTexture && InTexture->VirtualTextureStreaming != 0)
+	{
+		UE_LOG(LogLandscapePatch, Error, TEXT("ULandscapeTexturePatch::SetWeightPatchAlphaTextureAsset: Invalid texture. Virtual textures are not supported. Weightmap Layer Name %s."), *InWeightmapLayerName.ToString());
+		return;
+	}
+
+	// Set the default weightmap alpha texture for all weightmap patches
+	if (InWeightmapLayerName.IsNone())
+	{
+		Modify(bAlwaysMarkDirty);
+		WeightPatchesAlphaSettings.AlphaTextureAsset = InTexture;
+		return;
+	}
+
+	// Set the override weightmap texture for a specific weight patch
 	for (const TObjectPtr<ULandscapeWeightPatchTextureInfo>& WeightPatch : WeightPatches)
 	{
 		if (IsValid(WeightPatch) && WeightPatch->WeightmapLayerName == InWeightmapLayerName)
 		{
-			WeightPatch->bUseAlphaChannel = bUseAlphaChannel;
-			if (WeightPatch->InternalData)
+			WeightPatch->bOverrideWeightAlphaSettings = bInOverrideAlpha;
+			// When bInOverrideAlpha is false the weightmap layer will fallback to the shared weightmap alpha
+			if (bInOverrideAlpha)
+			{
+				WeightPatch->WeightPatchAlphaSettings.AlphaTextureAsset = InTexture;
+			}
+			else
 			{
-				WeightPatch->InternalData->SetUseAlphaChannel(bUseAlphaChannel);
+				WeightPatch->SetWeightPatchAlphaSourceMode(ELandscapeTexturePatchAlphaSourceMode::None, /*bInOverrideAlpha =*/false, bAlwaysMarkDirty);
 			}
 			return;
 		}
 	}
 	const FString LayerNameString = InWeightmapLayerName.ToString();
-	UE_LOG(LogLandscapePatch, Warning, TEXT("ULandscapeTexturePatch::SetUseAlphaChannelForWeightPatch: Unable to find data for weight layer %s"), *LayerNameString);
+	UE_LOG(LogLandscapePatch, Warning, TEXT("ULandscapeTexturePatch::SetWeightPatchAlphaTextureAsset: Unable to find data for weight layer %s"), *LayerNameString);
+}
+
+ELandscapeTexturePatchTextureChannel ULandscapeTexturePatch::GetWeightPatchAlphaTextureChannel(const FName& InWeightmapLayerName) const
+{
+	if (InWeightmapLayerName.IsNone())
+	{
+		return WeightPatchesAlphaSettings.AlphaTextureChannel;
+	}
+
+	for (const TObjectPtr<ULandscapeWeightPatchTextureInfo>& WeightPatch : WeightPatches)
+	{
+		if (IsValid(WeightPatch) && WeightPatch->WeightmapLayerName == InWeightmapLayerName)
+		{
+			// Returns the override value if and only if the alpha settings are overridden
+			return GetWeightPatchAlphaSettingsConst(WeightPatch).AlphaTextureChannel;
+		}
+	}
+
+	UE_LOG(LogLandscapePatch, Warning, TEXT("ULandscapeTexturePatch::GetWeightPatchAlphaTextureChannel: Unable to find data for weight layer %s"), *InWeightmapLayerName.ToString());
+	return WeightPatchesAlphaSettings.AlphaTextureChannel;
+}
+
+void ULandscapeTexturePatch::SetWeightPatchAlphaTextureChannel(const FName& InWeightmapLayerName, ELandscapeTexturePatchTextureChannel InTextureChannel, bool bInOverrideAlpha, bool bAlwaysMarkDirty)
+{
+	// Add to the default weightmap alpha channel for all weightmap patches
+	if (InWeightmapLayerName.IsNone())
+	{
+		if (WeightPatchesAlphaSettings.AlphaTextureChannel == InTextureChannel)
+		{
+			return;
+		}
+		Modify(bAlwaysMarkDirty);
+		WeightPatchesAlphaSettings.AlphaTextureChannel = InTextureChannel;
+		return;
+	}
+
+	for (const TObjectPtr<ULandscapeWeightPatchTextureInfo>& WeightPatch : WeightPatches)
+	{
+		if (IsValid(WeightPatch) && WeightPatch->WeightmapLayerName == InWeightmapLayerName)
+		{
+			WeightPatch->bOverrideWeightAlphaSettings = bInOverrideAlpha;
+			// When bInOverrideAlpha is false the weightmap layer will fallback to the shared weightmap channel
+			if (bInOverrideAlpha)
+			{
+				WeightPatch->SetWeightPatchAlphaTextureChannel(InTextureChannel, bAlwaysMarkDirty);
+			}
+			else
+			{
+				WeightPatch->SetWeightPatchAlphaSourceMode(ELandscapeTexturePatchAlphaSourceMode::None, /*bInOverrideAlpha =*/false, bAlwaysMarkDirty);
+			}
+			return;
+		}
+	}
+
+	UE_LOG(LogLandscapePatch, Warning, TEXT("ULandscapeTexturePatch::SetWeightPatchAlphaTextureChannel: Unable to set data for weight layer %s"), *InWeightmapLayerName.ToString());
 }
 
 void ULandscapeTexturePatch::SetWeightPatchSourceMode(const FName& InWeightmapLayerName, ELandscapeTexturePatchSourceMode NewMode)
@@ -1900,6 +2564,33 @@ ELandscapeTexturePatchSourceMode ULandscapeTexturePatch::GetWeightPatchSourceMod
 	return ELandscapeTexturePatchSourceMode::None;
 }
 
+ELandscapeTexturePatchTextureChannel ULandscapeTexturePatch::GetWeightPatchTextureChannel(const FName& InWeightmapLayerName) const
+{
+	for (const TObjectPtr<ULandscapeWeightPatchTextureInfo>& WeightPatch : WeightPatches)
+	{
+		if (IsValid(WeightPatch) && WeightPatch->WeightmapLayerName == InWeightmapLayerName)
+		{
+			return WeightPatch->WeightPatchTextureChannel;
+		}
+	}
+	UE_LOG(LogLandscapePatch, Warning, TEXT("ULandscapeTexturePatch::GetWeightPatchTextureChannel: Unable to get data for weight layer %s"), *InWeightmapLayerName.ToString());
+	return ELandscapeTexturePatchTextureChannel::None;
+}
+
+void ULandscapeTexturePatch::SetWeightPatchTextureChannel(const FName& InWeightmapLayerName, ELandscapeTexturePatchTextureChannel InTextureChannel, bool bAlwaysMarkDirty)
+{
+	for (const TObjectPtr<ULandscapeWeightPatchTextureInfo>& WeightPatch : WeightPatches)
+	{
+		if (IsValid(WeightPatch) && WeightPatch->WeightmapLayerName == InWeightmapLayerName)
+		{
+			WeightPatch->SetWeightPatchTextureChannel(InTextureChannel, bAlwaysMarkDirty);
+			return;
+		}
+	}
+
+	UE_LOG(LogLandscapePatch, Warning, TEXT("ULandscapeTexturePatch::SetWeightPatchTextureChannel: Unable to set data for weight layer %s"), *InWeightmapLayerName.ToString());
+}
+
 UTexture* ULandscapeTexturePatch::GetWeightPatchTextureAsset(const FName& InWeightmapLayerName)
 {
 	for (const TObjectPtr<ULandscapeWeightPatchTextureInfo>& WeightPatch : WeightPatches)
@@ -1957,6 +2648,34 @@ UTextureRenderTarget2D* ULandscapeTexturePatch::GetWeightPatchRenderTarget(ULand
 	return nullptr;
 }
 
+FLandscapeTexturePatchAlphaSettings& ULandscapeTexturePatch::GetWeightPatchAlphaSettings(ULandscapeWeightPatchTextureInfo* WeightPatch)
+{
+	if (WeightPatch == nullptr)
+	{
+		return WeightPatchesAlphaSettings;
+	}
+	return WeightPatch->bOverrideWeightAlphaSettings ? WeightPatch->WeightPatchAlphaSettings : WeightPatchesAlphaSettings;
+}
+
+const FLandscapeTexturePatchAlphaSettings& ULandscapeTexturePatch::GetWeightPatchAlphaSettingsConst(ULandscapeWeightPatchTextureInfo* WeightPatch) const
+{
+	if (WeightPatch == nullptr)
+	{
+		return WeightPatchesAlphaSettings;
+	}
+	return WeightPatch->bOverrideWeightAlphaSettings ? WeightPatch->WeightPatchAlphaSettings : WeightPatchesAlphaSettings;
+}
+
+UTextureRenderTarget2D* ULandscapeTexturePatch::GetWeightPatchAlphaRenderTarget(ULandscapeWeightPatchTextureInfo* WeightPatch)
+{
+	return GetWeightPatchAlphaSettings(WeightPatch).GetAlphaRenderTarget(WeightPatch);
+}
+
+UTexture2D* ULandscapeTexturePatch::GetWeightPatchAlphaInternalTexture(ULandscapeWeightPatchTextureInfo* WeightPatch)
+{
+	return GetWeightPatchAlphaSettings(WeightPatch).GetAlphaInternalTexture(WeightPatch);
+}
+
 UTexture2D* ULandscapeTexturePatch::GetWeightPatchInternalTexture(ULandscapeWeightPatchTextureInfo* WeightPatch)
 {
 #if WITH_EDITOR
@@ -1983,9 +2702,9 @@ UTexture2D* ULandscapeTexturePatch::GetWeightPatchInternalTexture(ULandscapeWeig
 
 void ULandscapeTexturePatch::SetWeightPatchTextureAsset(const FName& InWeightmapLayerName, UTexture* TextureIn)
 {
-	if (!ensureMsgf(!TextureIn || TextureIn->VirtualTextureStreaming == 0,
-		TEXT("ULandscapeTexturePatch::SetWeightPatchTextureAsset: Virtual textures are not supported.")))
+	if (TextureIn && TextureIn->VirtualTextureStreaming != 0)
 	{
+		UE_LOG(LogLandscapePatch, Error, TEXT("ULandscapeTexturePatch::SetWeightPatchTextureAsset: Invalid texture. Virtual textures are not supported. Weightmap Layer Name %s."), *InWeightmapLayerName.ToString());
 		return;
 	}
 
diff --git a/Engine/Plugins/Editor/LandscapePatch/Source/LandscapePatch/Private/LandscapeTexturePatchPS.cpp b/Engine/Plugins/Editor/LandscapePatch/Source/LandscapePatch/Private/LandscapeTexturePatchPS.cpp
--- a/Engine/Plugins/Editor/LandscapePatch/Source/LandscapePatch/Private/LandscapeTexturePatchPS.cpp
+++ b/Engine/Plugins/Editor/LandscapePatch/Source/LandscapePatch/Private/LandscapeTexturePatchPS.cpp
@@ -22,13 +22,24 @@ void FApplyLandscapeTextureHeightPatchPS::ModifyCompilationEnvironment(const FGl
 
 	// Make our flag choices match in the shader.
 	OutEnvironment.SetDefine(TEXT("RECTANGULAR_FALLOFF_FLAG"), static_cast<uint8>(FApplyLandscapeTextureHeightPatchPS::EFlags::RectangularFalloff));
-	OutEnvironment.SetDefine(TEXT("APPLY_PATCH_ALPHA_FLAG"), static_cast<uint8>(FApplyLandscapeTextureHeightPatchPS::EFlags::ApplyPatchAlpha));
 	OutEnvironment.SetDefine(TEXT("INPUT_IS_PACKED_HEIGHT_FLAG"), static_cast<uint8>(FApplyLandscapeTextureHeightPatchPS::EFlags::InputIsPackedHeight));
 
 	OutEnvironment.SetDefine(TEXT("ALPHA_BLEND_MODE"), static_cast<uint8>(ELandscapeTexturePatchBlendMode::AlphaBlend));
 	OutEnvironment.SetDefine(TEXT("ADDITIVE_MODE"), static_cast<uint8>(ELandscapeTexturePatchBlendMode::Additive));
 	OutEnvironment.SetDefine(TEXT("MIN_MODE"), static_cast<uint8>(ELandscapeTexturePatchBlendMode::Min));
 	OutEnvironment.SetDefine(TEXT("MAX_MODE"), static_cast<uint8>(ELandscapeTexturePatchBlendMode::Max));
+
+	OutEnvironment.SetDefine(TEXT("ALPHA_PATCH_NONE"), static_cast<uint8>(ELandscapeTexturePatchAlphaSourceMode::None));
+	OutEnvironment.SetDefine(TEXT("ALPHA_PATCH_INTERNAL_TEXTURE"), static_cast<uint8>(ELandscapeTexturePatchAlphaSourceMode::InternalTexture));
+	OutEnvironment.SetDefine(TEXT("ALPHA_PATCH_RENDER_TARGET"), static_cast<uint8>(ELandscapeTexturePatchAlphaSourceMode::TextureBackedRenderTarget));
+	OutEnvironment.SetDefine(TEXT("ALPHA_PATCH_TEXTURE_ASSET"), static_cast<uint8>(ELandscapeTexturePatchAlphaSourceMode::TextureAsset));
+	OutEnvironment.SetDefine(TEXT("ALPHA_PATCH_SRC_TEXTURE_CHANNEL"), static_cast<uint8>(ELandscapeTexturePatchAlphaSourceMode::SourceTextureChannel));
+
+	OutEnvironment.SetDefine(TEXT("TEXTURE_CHANNEL_R"), static_cast<uint8>(ELandscapeTexturePatchTextureChannel::Red));
+	OutEnvironment.SetDefine(TEXT("TEXTURE_CHANNEL_G"), static_cast<uint8>(ELandscapeTexturePatchTextureChannel::Green));
+	OutEnvironment.SetDefine(TEXT("TEXTURE_CHANNEL_B"), static_cast<uint8>(ELandscapeTexturePatchTextureChannel::Blue));
+	OutEnvironment.SetDefine(TEXT("TEXTURE_CHANNEL_A"), static_cast<uint8>(ELandscapeTexturePatchTextureChannel::Alpha));
+	OutEnvironment.SetDefine(TEXT("TEXTURE_CHANNEL_LUM"), static_cast<uint8>(ELandscapeTexturePatchTextureChannel::Luminance));
 }
 
 void FApplyLandscapeTextureHeightPatchPS::AddToRenderGraph(FRDGEventName&& InRDGEventName, FRDGBuilder& GraphBuilder, FParameters* InParameters, const FIntRect& DestinationBounds)
@@ -195,12 +206,23 @@ void FApplyLandscapeTextureWeightPatchPS::ModifyCompilationEnvironment(const FGl
 
 	// Make our flag choices match in the shader.
 	OutEnvironment.SetDefine(TEXT("RECTANGULAR_FALLOFF_FLAG"), static_cast<uint8>(FApplyLandscapeTextureWeightPatchPS::EFlags::RectangularFalloff));
-	OutEnvironment.SetDefine(TEXT("APPLY_PATCH_ALPHA_FLAG"), static_cast<uint8>(FApplyLandscapeTextureWeightPatchPS::EFlags::ApplyPatchAlpha));
 
 	OutEnvironment.SetDefine(TEXT("ALPHA_BLEND_MODE"), static_cast<uint8>(ELandscapeTexturePatchBlendMode::AlphaBlend));
 	OutEnvironment.SetDefine(TEXT("ADDITIVE_MODE"), static_cast<uint8>(ELandscapeTexturePatchBlendMode::Additive));
 	OutEnvironment.SetDefine(TEXT("MIN_MODE"), static_cast<uint8>(ELandscapeTexturePatchBlendMode::Min));
 	OutEnvironment.SetDefine(TEXT("MAX_MODE"), static_cast<uint8>(ELandscapeTexturePatchBlendMode::Max));
+
+	OutEnvironment.SetDefine(TEXT("ALPHA_PATCH_NONE"), static_cast<uint8>(ELandscapeTexturePatchAlphaSourceMode::None));
+	OutEnvironment.SetDefine(TEXT("ALPHA_PATCH_INTERNAL_TEXTURE"), static_cast<uint8>(ELandscapeTexturePatchAlphaSourceMode::InternalTexture));
+	OutEnvironment.SetDefine(TEXT("ALPHA_PATCH_RENDER_TARGET"), static_cast<uint8>(ELandscapeTexturePatchAlphaSourceMode::TextureBackedRenderTarget));
+	OutEnvironment.SetDefine(TEXT("ALPHA_PATCH_TEXTURE_ASSET"), static_cast<uint8>(ELandscapeTexturePatchAlphaSourceMode::TextureAsset));
+	OutEnvironment.SetDefine(TEXT("ALPHA_PATCH_SRC_TEXTURE_CHANNEL"), static_cast<uint8>(ELandscapeTexturePatchAlphaSourceMode::SourceTextureChannel));
+
+	OutEnvironment.SetDefine(TEXT("TEXTURE_CHANNEL_R"), static_cast<uint8>(ELandscapeTexturePatchTextureChannel::Red));
+	OutEnvironment.SetDefine(TEXT("TEXTURE_CHANNEL_G"), static_cast<uint8>(ELandscapeTexturePatchTextureChannel::Green));
+	OutEnvironment.SetDefine(TEXT("TEXTURE_CHANNEL_B"), static_cast<uint8>(ELandscapeTexturePatchTextureChannel::Blue));
+	OutEnvironment.SetDefine(TEXT("TEXTURE_CHANNEL_A"), static_cast<uint8>(ELandscapeTexturePatchTextureChannel::Alpha));
+	OutEnvironment.SetDefine(TEXT("TEXTURE_CHANNEL_LUM"), static_cast<uint8>(ELandscapeTexturePatchTextureChannel::Luminance));
 }
 
 void FApplyLandscapeTextureWeightPatchPS::AddToRenderGraph(FRDGEventName&& InRDGEventName, FRDGBuilder& GraphBuilder, FParameters* InParameters, const FIntRect& DestinationBounds)
diff --git a/Engine/Plugins/Editor/LandscapePatch/Source/LandscapePatch/Private/LandscapeTexturePatchPS.h b/Engine/Plugins/Editor/LandscapePatch/Source/LandscapePatch/Private/LandscapeTexturePatchPS.h
--- a/Engine/Plugins/Editor/LandscapePatch/Source/LandscapePatch/Private/LandscapeTexturePatchPS.h
+++ b/Engine/Plugins/Editor/LandscapePatch/Source/LandscapePatch/Private/LandscapeTexturePatchPS.h
@@ -31,12 +31,9 @@ namespace UE::Landscape
 			// When false, falloff is circular.
 			RectangularFalloff = 1 << 0,
 
-			// When true, the texture alpha channel is considered for blending (in addition to falloff, if nonzero)
-			ApplyPatchAlpha = 1 << 1,
-
 			// When false, the input is directly interpreted as being the height value to process. When true, the height
 			// is unpacked from the red and green channels to make a 16 bit int.
-			InputIsPackedHeight = 1 << 2
+			InputIsPackedHeight = 1 << 1
 		};
 
 		// TODO: We could consider exposing an additional global alpha setting that we can use to pass in the given
@@ -64,6 +61,12 @@ namespace UE::Landscape
 			SHADER_PARAMETER(uint32, InBlendMode)
 			// Some combination of the flags (see constants above).
 			SHADER_PARAMETER(uint32, InFlags)
+			// Height source texture color channel
+			SHADER_PARAMETER(uint32, InTextureChannel)
+			// Alpha source mode, override texture (White system texture if undefined), and texture alpha color channel
+			SHADER_PARAMETER(uint32, InAlphaMode)
+			SHADER_PARAMETER_RDG_TEXTURE_SRV(Texture2D<float4>, InHeightPatchAlphaTexture)
+			SHADER_PARAMETER(uint32, InAlphaTextureChannel)
 
 			RENDER_TARGET_BINDING_SLOTS() // Holds our output
 		END_SHADER_PARAMETER_STRUCT()
@@ -169,9 +172,6 @@ namespace UE::Landscape
 
 			// When false, falloff is circular.
 			RectangularFalloff = 1 << 0,
-
-			// When true, the texture alpha channel is considered for blending (in addition to falloff, if nonzero)
-			ApplyPatchAlpha = 1 << 1,
 		};
 
 		// TODO: We could consider exposing an additional global alpha setting that we can use to pass in the given
@@ -193,6 +193,12 @@ namespace UE::Landscape
 			SHADER_PARAMETER(uint32, InBlendMode)
 			// Some combination of the flags (see constants above).
 			SHADER_PARAMETER(uint32, InFlags)
+			// Weight source texture color channel
+			SHADER_PARAMETER(uint32, InTextureChannel)
+			// Alpha source mode, override texture (White system texture if undefined), and alpha texture color channel
+			SHADER_PARAMETER(uint32, InAlphaMode)
+			SHADER_PARAMETER_RDG_TEXTURE_SRV(Texture2D<float4>, InWeightPatchAlphaTexture)
+			SHADER_PARAMETER(uint32, InAlphaTextureChannel)
 
 			RENDER_TARGET_BINDING_SLOTS() // Holds our output
 		END_SHADER_PARAMETER_STRUCT()
diff --git a/Engine/Plugins/Editor/LandscapePatch/Source/LandscapePatch/Public/LandscapeTexturePatch.h b/Engine/Plugins/Editor/LandscapePatch/Source/LandscapePatch/Public/LandscapeTexturePatch.h
--- a/Engine/Plugins/Editor/LandscapePatch/Source/LandscapePatch/Public/LandscapeTexturePatch.h
+++ b/Engine/Plugins/Editor/LandscapePatch/Source/LandscapePatch/Public/LandscapeTexturePatch.h
@@ -54,6 +54,61 @@ enum class ELandscapeTexturePatchSourceMode : uint8
 	 TextureAsset
 };
 
+/**
+ * Determines the alpha mask applied to a patch texture. Alpha is applied in addition to the patch falloff settings
+ */
+UENUM(BlueprintType)
+enum class ELandscapeTexturePatchAlphaSourceMode : uint8
+{
+	/**
+	 * The patch texture is not masked by any alpha value
+	 */
+	None,
+
+	/**
+	 * The alpha data will be read from an internally-stored UTexture2D. In this mode, the patch can't be written-to via
+	 * blueprints, but it avoids storing the extra render target needed for TextureBackedRenderTarget.
+	 */
+	InternalTexture,
+
+	/**
+	* The patch alpha data will be read from an internally-stored render target, which can be written to via Blueprints
+	* and which gets serialized to an internally stored UTexture2D when needed. Uses double the memory of InternalTexture.
+	*/
+	TextureBackedRenderTarget,
+
+	/**
+	 * The alpha data will be read from a UTexture asset (which can be a render target). Allows multiple patches
+	 * to share the same alpha texture.
+	 */
+	TextureAsset,
+
+	/**
+	 * A color channel of the original source texture will be used as the alpha mask. By default, the alpha channel of the
+	 * texture is used. The source texture must have data in the defined channel to affect the patch
+	 */
+	SourceTextureChannel
+};
+
+/**
+ * Color channel of the source texture
+ */
+UENUM(BlueprintType)
+enum class ELandscapeTexturePatchTextureChannel : uint8
+{
+	None UMETA(Hidden),
+	// Texture's red channel
+	Red,
+	// Texture's green channel
+	Green,
+	// Texture's blue channel
+	Blue, 
+	// Texture's alpha channel
+	Alpha, 
+	// Texture's luminance 
+	Luminance
+};
+
 // Determines how the patch is combined with the previous state of the landscape.
 UENUM(BlueprintType)
 enum class ELandscapeTexturePatchBlendMode : uint8
@@ -127,6 +182,56 @@ enum class ELandscapeTextureHeightPatchZeroHeightMeaning : uint8
 	WorldZero
 };
 
+USTRUCT(BlueprintType)
+struct FLandscapeTexturePatchAlphaSettings
+{
+	GENERATED_BODY()
+public:
+	/** When alpha source mode is TextureBackedRenderTarget returns the alpha render target. Will allocate new render target if not yet initialized */
+	UTextureRenderTarget2D* GetAlphaRenderTarget(UObject* InOuter);
+	/** When alpha source mode is  InternalTexture/TextureBackedRenderTarget returns the alpha internal texture. Will allocate new internal data if not yet initialized */
+	UTexture2D* GetAlphaInternalTexture(UObject* InOuter);
+
+	/** Transitions from the previous alpha source mode to the new mode. Marks the Outer owner dirty and clears internal data if needed  
+	* @param InNewSourceMode the new alpha source mode
+	* @param InOuter The alpha settings owner (must be direct texture patch or weight patch info)
+	* @param bAlwaysMarkDirty When true, marks the outer package dirty when the source mode changes
+	*/
+	void SetAlphaSourceMode(ELandscapeTexturePatchAlphaSourceMode InNewSourceMode, UObject* InOuter, bool bAlwaysMarkDirty = true);
+	ELandscapeTexturePatchAlphaSourceMode GetAlphaSourceMode() const
+	{
+		return AlphaSourceMode;
+	};
+
+	/** How the alpha texture is stored. */
+	UPROPERTY(EditAnywhere, Category = "Texture Alpha Settings", meta = (DisplayName = "Source Mode"))
+	ELandscapeTexturePatchAlphaSourceMode DetailPanelAlphaSourceMode = ELandscapeTexturePatchAlphaSourceMode::None;
+
+	/** Alpha Texture to use when the source mode is set to texture asset. */
+	UPROPERTY(EditAnywhere, Category = "Texture Alpha Settings", meta = (EditConditionHides,
+		EditCondition = "AlphaSourceMode == ELandscapeTexturePatchAlphaSourceMode::TextureAsset", 
+		DisallowedAssetDataTags = "VirtualTextureStreaming=True"))
+	TObjectPtr<UTexture> AlphaTextureAsset;
+
+	/** Texture channel utilized for the alpha mask. The texture alpha channel is used by default. */
+	UPROPERTY(EditAnywhere, Category = "Texture Alpha Settings", meta = (EditConditionHides, EditCondition = "AlphaSourceMode != ELandscapeTexturePatchAlphaSourceMode::None"))
+	ELandscapeTexturePatchTextureChannel AlphaTextureChannel = ELandscapeTexturePatchTextureChannel::Alpha;
+
+	/** Not directly settable via detail panel- for alpha display/debugging purposes only. */
+	UPROPERTY(VisibleAnywhere, Category = "Texture Alpha Settings", Instanced, AdvancedDisplay, meta = (DisplayName = "Debug Data", EditConditionHides, EditCondition = "AlphaSourceMode != ELandscapeTexturePatchAlphaSourceMode::None"))
+	TObjectPtr<ULandscapeWeightTextureBackedRenderTarget> InternalAlphaData;
+
+	// TODO [jared.ritchie] Add Scale/Offset/Rotation UV properties here
+
+private:
+	/** When the source mode changes, there are transition steps that need run to clean up render targets / internal data.
+	* In PostEditChangeProperty, the previous source mode is already lost and it becomes harder to track the necessary transition.
+	* DetailPanelAlphaSourceMode exposes the property and allows us to process the source mode transition before updating AlphaSourceMode
+	*/
+	UPROPERTY(VisibleAnywhere, Category = "Texture Alpha Settings", meta = (EditConditionHides, EditCondition = "false"))
+	ELandscapeTexturePatchAlphaSourceMode AlphaSourceMode = ELandscapeTexturePatchAlphaSourceMode::None;
+};
+
 //~ A struct in case we find that we need other encoding settings.
 USTRUCT(BlueprintType)
 struct FLandscapeTexturePatchEncodingSettings
@@ -164,6 +269,7 @@ public:
 #if WITH_EDITOR
 	// UObject
 	LANDSCAPEPATCH_API virtual void PostEditChangeProperty(FPropertyChangedEvent& PropertyChangedEvent) override;
+	LANDSCAPEPATCH_API virtual void PostEditChangeChainProperty(FPropertyChangedChainEvent& PropertyChangedEvent) override;
 	LANDSCAPEPATCH_API virtual void PreDuplicate(FObjectDuplicationParameters& DupParams) override;
 #endif // WITH_EDITOR
 
@@ -176,33 +282,41 @@ protected:
 	UPROPERTY(EditAnywhere, Category = WeightPatch)
 	bool bEditVisibilityLayer = false;
 	
+	UPROPERTY(EditAnywhere, Category = WeightPatch, meta = (EditConditionHides, EditCondition = "false"))
+	ELandscapeTexturePatchSourceMode SourceMode = ELandscapeTexturePatchSourceMode::None;
+
+	/** How the weightmap of the patch is stored. */
+	UPROPERTY(EditAnywhere, Category = WeightPatch, meta = (DisplayName = "Source Mode"))
+	ELandscapeTexturePatchSourceMode DetailPanelSourceMode = ELandscapeTexturePatchSourceMode::None;
+
 	/** Texture to use when source mode is set to texture asset. */
 	UPROPERTY(EditAnywhere, Category = WeightPatch, meta = (EditConditionHides,
 		EditCondition = "SourceMode == ELandscapeTexturePatchSourceMode::TextureAsset", 
 		DisallowedAssetDataTags = "VirtualTextureStreaming=True"))
 	TObjectPtr<UTexture> TextureAsset = nullptr;
 
+	/** Color channel of the texture used for weight data. By default the weight texture's red channel is used. */
+	UPROPERTY(EditAnywhere, Category = HeightPatch, meta = (EditConditionHides, EditCondition = "SourceMode != ELandscapeTexturePatchSourceMode::None"))
+	ELandscapeTexturePatchTextureChannel WeightPatchTextureChannel = ELandscapeTexturePatchTextureChannel::Red;
+
 	/** Not directly settable via detail panel- for display/debugging purposes only. */
 	UPROPERTY(VisibleAnywhere, Category = WeightPatch, Instanced, AdvancedDisplay, meta = (DisplayName = "Debug Data"))
 	TObjectPtr<ULandscapeWeightTextureBackedRenderTarget> InternalData = nullptr;
 
-	UPROPERTY(EditAnywhere, Category = WeightPatch, meta = (EditConditionHides, EditCondition = "false"))
-	ELandscapeTexturePatchSourceMode SourceMode = ELandscapeTexturePatchSourceMode::None;
+	UE_DEPRECATED(5.8, "bUseAlphaChannel is deprecated, use the WeightPatchAlphaSettings to define the alpha behavior.")
+	UPROPERTY(meta = (DeprecatedProperty, DeprecationMessage = "Use WeightPatchAlphaSettings instead. See ELandscapeTexturePatchAlphaSourceMode::SourceTextureChannel to enable legacy behavior"))
+	bool bUseAlphaChannel_DEPRECATED = false;
 
-	/** How the weightmap of the patch is stored. */
-	UPROPERTY(EditAnywhere, Category = WeightPatch, meta = (DisplayName = "Source Mode"))
-	ELandscapeTexturePatchSourceMode DetailPanelSourceMode = ELandscapeTexturePatchSourceMode::None;
+	UPROPERTY(EditAnywhere, Category = WeightPatch, meta = (InlineEditConditionToggle))
+	bool bOverrideWeightAlphaSettings = false;
 
-	/**
-	 * When true, texture alpha channel will be used when applying the weight patch. Note that the source data needs to
-	 * have an alpha channel for this to have an effect.
-	 */
-	UPROPERTY(EditAnywhere, Category = WeightPatch)
-	bool bUseAlphaChannel = false;
+	UPROPERTY(EditAnywhere, Category = WeightPatch, meta = (DisplayName = "Override Alpha Settings", EditCondition = "bOverrideWeightAlphaSettings"))
+	FLandscapeTexturePatchAlphaSettings WeightPatchAlphaSettings;
 
-	/** Allows to override the patch component's own blend node */
+	/** Allows to override the patch component's weightmap blend node */
 	UPROPERTY(EditAnywhere, Category = WeightPatch, meta = (InlineEditConditionToggle))
 	bool bOverrideBlendMode = false;
+
 	/** Determines how the patch is combined with the previous state of the landscape (overrides the patch component's own blend node). */
 	UPROPERTY(EditAnywhere, Category = WeightPatch, meta = (EditCondition = "bOverrideBlendMode"))
 	ELandscapeTexturePatchBlendMode OverrideBlendMode = ELandscapeTexturePatchBlendMode::AlphaBlend;
@@ -213,6 +327,10 @@ protected:
 	bool bReinitializeOnNextRender = false;
 
 	void SetSourceMode(ELandscapeTexturePatchSourceMode NewMode);
+	void SetWeightPatchAlphaSourceMode(ELandscapeTexturePatchAlphaSourceMode InSourceMode, bool bInOverrideAlpha, bool bAlwaysMarkDirty = true);
+	void SetWeightPatchTextureChannel(ELandscapeTexturePatchTextureChannel InTextureChannel, bool bAlwaysMarkDirty = true);
+	void SetWeightPatchAlphaTextureChannel(ELandscapeTexturePatchTextureChannel InTextureChannel, bool bAlwaysMarkDirty = true);
+
 #if WITH_EDITOR
 	void TransitionSourceModeInternal(ELandscapeTexturePatchSourceMode OldMode, ELandscapeTexturePatchSourceMode NewMode);
 #endif
@@ -253,6 +371,7 @@ public:
 	LANDSCAPEPATCH_API virtual void CheckForErrors() override;
 
 	// UObject
+	LANDSCAPEPATCH_API virtual void PostLoad() override;
 	LANDSCAPEPATCH_API virtual void PostEditChangeProperty(FPropertyChangedEvent& PropertyChangedEvent) override;
 #endif
 	LANDSCAPEPATCH_API virtual void Serialize(FArchive& Ar) override;
@@ -370,6 +489,39 @@ public:
 	UFUNCTION(BlueprintCallable, Category = LandscapePatch)
 	LANDSCAPEPATCH_API virtual void SetHeightSourceMode(ELandscapeTexturePatchSourceMode NewMode);
 
+	/**
+	 * Sets the height texture's color channel. Texture channel is restricted when HeightSource is InternalData or HeightEncoding is NativePackedHeight (Assumes RG channels)
+	 * @param InTextureChannel the color channel of the texture to use
+	 * @param bAlwaysMarkDirty mark the package dirty when channel property changes
+	 */
+	UFUNCTION(BlueprintCallable, Category = LandscapePatch)
+	LANDSCAPEPATCH_API virtual void SetHeightTextureChannel(ELandscapeTexturePatchTextureChannel InTextureChannel, bool bAlwaysMarkDirty = true);
+
+	UFUNCTION(BlueprintCallable, Category = LandscapePatch)
+	virtual ELandscapeTexturePatchTextureChannel GetHeightTextureChannel() const { return HeightTextureChannel; }
+
+	UFUNCTION(BlueprintCallable, Category = LandscapePatch)
+	virtual ELandscapeTexturePatchAlphaSourceMode GetHeightAlphaSourceMode() const { return HeightAlphaSettings.GetAlphaSourceMode(); }
+
+	/**
+	 * Changes the source mode of the height patch alpha data. There are currently no API guarantees regarding the initialization of the
+	 * new alpha source data. E.g. when first switching to use an internal render target, the data in that
+	 * render target may not be initialized.
+	 */
+	UFUNCTION(BlueprintCallable, Category = LandscapePatch)
+	LANDSCAPEPATCH_API virtual void SetHeightAlphaSourceMode(ELandscapeTexturePatchAlphaSourceMode InNewMode, bool bAlwaysMarkDirty = true);
+
+	UFUNCTION(BlueprintCallable, Category = LandscapePatch)
+	virtual ELandscapeTexturePatchTextureChannel GetHeightAlphaSourceTextureChannel() const { return HeightAlphaSettings.AlphaTextureChannel; }
+
+	 /**
+	 * Sets the height alpha texture patch color channel
+	 * @param InTextureChannel the color channel of the texture to use for the height alpha
+	 * @param bAlwaysMarkDirty mark the package dirty when channel property changes
+	 */
+	UFUNCTION(BlueprintCallable, Category = LandscapePatch)
+	LANDSCAPEPATCH_API virtual void SetHeightAlphaTextureChannel(ELandscapeTexturePatchTextureChannel InTextureChannel, bool bAlwaysMarkDirty = true);
+
 	/**
 	 * Sets the texture used for height when the height source mode is set to texture asset. Note that
 	 * virtual textures are not supported.
@@ -402,19 +554,12 @@ public:
 	UFUNCTION(BlueprintCallable, Category = LandscapePatch, meta = (ETextureRenderTargetFormat = "ETextureRenderTargetFormat::RTF_R32f"))
 	LANDSCAPEPATCH_API void SetHeightRenderTargetFormat(ETextureRenderTargetFormat Format);
 
-	/**
-	 * Determines whether the height patch alpha channel is used for blending into the existing values.
-	 * Note that the source data needs to have an alpha channel in this case. How the alpha channel is
-	 * used depends on the patch blend mode (see SetBlendMode).
-	 */
-	UFUNCTION(BlueprintCallable, Category = "LandscapePatch")
+	UE_DEPRECATED(all, "SetUseAlphaChannelForHeight is deprecated. Use SetHeightAlphaSourceMode with ELandscapeTexturePatchAlphaSourceMode instead")
+	UFUNCTION(meta = (DeprecatedFunction, DeprecationMessage = "SetUseAlphaChannelForHeight has been deprecated, use SetHeightAlphaSourceMode instead."))
 	void SetUseAlphaChannelForHeight(bool bUse)
 	{
-		if (bUseTextureAlphaForHeight != bUse)
-		{
-			Modify();
-			bUseTextureAlphaForHeight = bUse;
-		}
+		// Deprecated
+		SetHeightAlphaSourceMode(bUse ? ELandscapeTexturePatchAlphaSourceMode::SourceTextureChannel : ELandscapeTexturePatchAlphaSourceMode::None);
 	}
 
 	/**
@@ -428,6 +573,12 @@ public:
 		{
 			Modify();
 			HeightEncoding = EncodingMode;
+
+			// When the height uses internal texture / native packed height the texture channel is restricted (assume RG texture channels)
+			if (HeightEncoding == ELandscapeTextureHeightPatchEncoding::NativePackedHeight)
+			{
+				HeightTextureChannel = ELandscapeTexturePatchTextureChannel::None;
+			}
 		}
 	}
 
@@ -461,12 +612,15 @@ public:
 
 
 	// Weight related functions:
+	UE_DEPRECATED(all, "Use CreateWeightPatch with ELandscapeTexturePatchAlphaSourceMode instead")
+	UFUNCTION(meta = (DeprecatedFunction, DeprecationMessage = "AddWeightPatch has been deprecated, use CreateWeightPatch with ELandscapeTexturePatchAlphaSourceMode instead."))
+	LANDSCAPEPATCH_API virtual void AddWeightPatch(const FName& InWeightmapLayerName, ELandscapeTexturePatchSourceMode SourceMode, bool bUseAlphaChannel); 
 
 	/**
 	 * By default, the layer is added with source mode set to be a texture-backed render target.
 	 */
 	UFUNCTION(BlueprintCallable, Category = LandscapePatch)
-	LANDSCAPEPATCH_API virtual void AddWeightPatch(const FName& InWeightmapLayerName, ELandscapeTexturePatchSourceMode SourceMode, bool bUseAlphaChannel);
+	LANDSCAPEPATCH_API virtual void CreateWeightPatch(const FName& InWeightmapLayerName, ELandscapeTexturePatchSourceMode InSourceMode, ELandscapeTexturePatchAlphaSourceMode InAlphaSourceMode);
 
 	UFUNCTION(BlueprintCallable, Category = LandscapePatch)
 	LANDSCAPEPATCH_API virtual void RemoveWeightPatch(const FName& InWeightmapLayerName);
@@ -490,6 +644,18 @@ public:
 	UFUNCTION(BlueprintCallable, Category = LandscapePatch)
 	LANDSCAPEPATCH_API virtual void SetWeightPatchSourceMode(const FName& InWeightmapLayerName, ELandscapeTexturePatchSourceMode NewMode);
 
+	UFUNCTION(BlueprintCallable, Category = LandscapePatch)
+	LANDSCAPEPATCH_API virtual ELandscapeTexturePatchTextureChannel GetWeightPatchTextureChannel(const FName& InWeightmapLayerName) const;
+
+	/**
+	 * Sets the texture patch color channel specific weight patch
+	 * @param InWeightmapLayerName weightmap layer name to set the texture channel
+	 * @param InTextureChannel the channel of the texture the weight patch will use
+	 * @param bAlwaysMarkDirty mark the package dirty when property changes
+	 */
+	UFUNCTION(BlueprintCallable, Category = LandscapePatch)
+	LANDSCAPEPATCH_API virtual void SetWeightPatchTextureChannel(const FName& InWeightmapLayerName, ELandscapeTexturePatchTextureChannel InTextureChannel, bool bAlwaysMarkDirty = true);
+
 	UFUNCTION(BlueprintCallable, Category = LandscapePatch)
 	LANDSCAPEPATCH_API void SetWeightPatchTextureAsset(const FName& InWeightmapLayerName, UTexture* TextureIn);
 
@@ -500,8 +666,47 @@ public:
 	UFUNCTION(BlueprintCallable, Category = LandscapePatch)
 	LANDSCAPEPATCH_API virtual UTextureRenderTarget2D* GetWeightPatchRenderTarget(const FName& InWeightmapLayerName, bool bMarkDirty = true);
 
+	UE_DEPRECATED(all, "SetUseAlphaChannelForWeightPatch is deprecated. Use SetWeightPatchAlphaSourceMode with ELandscapeTexturePatchAlphaSourceMode instead")
+	UFUNCTION(meta = (DeprecatedFunction, DeprecationMessage = "SetUseAlphaChannelForWeightPatch has been deprecated, use SetWeightPatchAlphaSourceMode instead."))
+	LANDSCAPEPATCH_API virtual void SetUseAlphaChannelForWeightPatch(const FName& InWeightmapLayerName, bool bUseAlphaChannel); 
+
+	UFUNCTION(BlueprintCallable, Category = LandscapePatch)
+	LANDSCAPEPATCH_API virtual ELandscapeTexturePatchAlphaSourceMode GetWeightPatchAlphaSourceMode(const FName& InWeightmapLayerName) const;
+
+	/**
+	 * Sets the default weight patch alpha source mode for all weight patches or overrides the source for a specific weight patch
+	 * @param InWeightmapLayerName weightmap layer name to set the alpha mode. Overrides a specific layer or sets default source mode when name is None
+	 * @param InSourceMode the alpha source mode
+	 * @param bInOverrideAlpha sets the weightmap layer's override flag. When false, the weightmap layer will fallback to the default source mode
+ 	 * @param bAlwaysMarkDirty mark the package dirty when property changes
+	 */
+	LANDSCAPEPATCH_API virtual void SetWeightPatchAlphaSourceMode(const FName& InWeightmapLayerName, ELandscapeTexturePatchAlphaSourceMode InSourceMode, bool bInOverrideAlpha = true, bool bAlwaysMarkDirty = true);
+
+	UFUNCTION(BlueprintCallable, Category = LandscapePatch)
+	LANDSCAPEPATCH_API virtual const UTexture* GetWeightPatchAlphaTextureAsset(const FName& InWeightmapLayerName) const;
+
+	/**
+	 * Sets the default weight patch alpha texture asset for all weight patches or overrides the texture for a specific weight patch
+	 * @param InWeightmapLayerName weightmap layer name to set the alpha texture. Overrides a specific layer or sets the default texture when name is None
+	 * @param InTexture the alpha texture the weightmap will use
+	 * @param bInOverrideAlpha sets the weightmap layer's override flag. When false, the weightmap layer will fallback to the default alpha texture
+	 * @param bAlwaysMarkDirty mark the package dirty when property changes
+	 */
+	UFUNCTION(BlueprintCallable, Category = LandscapePatch)
+	LANDSCAPEPATCH_API void SetWeightPatchAlphaTextureAsset(const FName& InWeightmapLayerName, UTexture* InTexture, bool bInOverrideAlpha = true, bool bAlwaysMarkDirty = true);
+
+	UFUNCTION(BlueprintCallable, Category = LandscapePatch)
+	LANDSCAPEPATCH_API virtual ELandscapeTexturePatchTextureChannel GetWeightPatchAlphaTextureChannel(const FName& InWeightmapLayerName) const;
+
+	/**
+	 * Sets the texture patch color channel for all weight patches or set the texture channel for a specific weight patch alpha override
+	 * @param InWeightmapLayerName weightmap layer name to set the alpha texture channel. Overrides a specific layer or sets to the default texture channel when name is None
+	 * @param InTextureChannel the channel of the texture the alpha mask will use
+	 * @param bInOverrideAlpha sets the weightmap layer's override flag. When false, the weightmap layer will fallback to the default alpha channel
+	 * @param bAlwaysMarkDirty mark the package dirty when property changes
+	 */
 	UFUNCTION(BlueprintCallable, Category = LandscapePatch)
-	LANDSCAPEPATCH_API virtual void SetUseAlphaChannelForWeightPatch(const FName& InWeightmapLayerName, bool bUseAlphaChannel);
+	LANDSCAPEPATCH_API virtual void SetWeightPatchAlphaTextureChannel(const FName& InWeightmapLayerName, ELandscapeTexturePatchTextureChannel InTextureChannel, bool bInOverrideAlpha = true, bool bAlwaysMarkDirty = true);
 
 	UFUNCTION(BlueprintCallable, Category = LandscapePatch)
 	LANDSCAPEPATCH_API virtual void SetWeightPatchBlendModeOverride(const FName& InWeightmapLayerName, ELandscapeTexturePatchBlendMode BlendMode);
@@ -526,7 +731,7 @@ protected:
 	UPROPERTY(EditAnywhere, Category = Settings, meta = (UIMin = "0", ClampMin = "0"))
 	FVector2D UnscaledPatchCoverage = FVector2D(2000, 2000);
 
-	/** Determines how the patch is combined with the previous state of the landscape. */
+	/** Determines how the patch's heightmap is combined with the previous state of the landscape. */
 	UPROPERTY(EditAnywhere, Category = Settings)
 	ELandscapeTexturePatchBlendMode BlendMode = ELandscapeTexturePatchBlendMode::AlphaBlend;
 
@@ -556,6 +761,14 @@ protected:
 	UPROPERTY(EditAnywhere, Category = HeightPatch, meta = (DisplayName = "Source Mode"))
 	ELandscapeTexturePatchSourceMode DetailPanelHeightSourceMode = ELandscapeTexturePatchSourceMode::None;
 
+	/** Texture's color channel used for the height. The texture's Red channel is selected by default. 
+	* The patch will not affect the landscape when no channel is selected.
+	* When the height source mode is InternalTexture or HeightEncoding is NativePackedHeight the texture channel cannot be customized (assumes RG channel)
+	*/
+	UPROPERTY(EditAnywhere, Category = HeightPatch, meta = (DisplayName = "Texture Channels", EditConditionHides, 
+		EditCondition = "HeightSourceMode != ELandscapeTexturePatchSourceMode::None && HeightSourceMode != ELandscapeTexturePatchSourceMode::InternalTexture && HeightEncoding != ELandscapeTextureHeightPatchEncoding::NativePackedHeight"))
+	ELandscapeTexturePatchTextureChannel HeightTextureChannel = ELandscapeTexturePatchTextureChannel::Red;
+
 	/** Not directly settable via detail panel- for display/debugging purposes only. */
 	UPROPERTY(VisibleAnywhere, Category = HeightPatch, Instanced, AdvancedDisplay, meta = (DisplayName = "Debug Data"))
 	TObjectPtr<ULandscapeHeightTextureBackedRenderTarget> HeightInternalData = nullptr;
@@ -568,13 +781,12 @@ protected:
 		DisallowedAssetDataTags = "VirtualTextureStreaming=True"))
 	TObjectPtr<UTexture> HeightTextureAsset = nullptr;
 
-	
-	/** 
-	 * When true, texture alpha channel will be used when applying the patch. Note that the source data needs to
-	 * have an alpha channel for this to have an effect.
-	 */
-	UPROPERTY(EditAnywhere, Category = HeightPatch)
-	bool bUseTextureAlphaForHeight = false;
+	UE_DEPRECATED(5.8, "bUseAlphaChannel is deprecated, use the ELandscapeTexturePatchAlphaSourceMode to define the alpha behavior.")
+	UPROPERTY(meta = (DeprecatedProperty, DeprecationMessage = "Use HeightAlphaSettings instead. See ELandscapeTexturePatchAlphaSourceMode::SourceTextureChannel to enable legacy behavior"))
+	bool bUseTextureAlphaForHeight_DEPRECATED = false;
+
+	UPROPERTY(EditAnywhere, Category = HeightPatch, meta = (EditConditionHides, EditCondition = "HeightSourceMode != ELandscapeTexturePatchSourceMode::None"))
+	FLandscapeTexturePatchAlphaSettings HeightAlphaSettings;
 
 	/** How the values stored in the patch represent the height. Not customizable for Internal Texture source mode, which always uses native packed height. */
 	UPROPERTY(EditAnywhere, Category = HeightPatch, meta = (
@@ -612,6 +824,18 @@ protected:
 	UPROPERTY(EditAnywhere, Category = WeightPatches, Instanced, NoClear, meta=(NoResetToDefault))
 	TArray<TObjectPtr<ULandscapeWeightPatchTextureInfo>> WeightPatches;
 
+	/* Default weight patches alpha settings.
+	* Defines the alpha mask settings for all weight patches within this texture patch. Alpha is only applied
+	* to weight patches with a valid source mode. Alpha settings can be overridden by individual weight patches.
+	*/
+	UPROPERTY(EditAnywhere, Category = WeightPatches)
+	FLandscapeTexturePatchAlphaSettings WeightPatchesAlphaSettings;
+
+	/** Determines how the patch's weightmap layers are combined with the previous state of the landscape. Blend
+	* mode can be overriden by individual weight patches */
+	UPROPERTY(EditAnywhere, Category = WeightPatches)
+	ELandscapeTexturePatchBlendMode WeightPatchesBlendMode = ELandscapeTexturePatchBlendMode::AlphaBlend;
+
 	// Reinitialization from detail panel:
 
 	/**
@@ -658,11 +882,25 @@ protected:
 	int32 InitTextureSizeY = 33;
 
 	LANDSCAPEPATCH_API UTexture2D* GetHeightInternalTexture();
+	LANDSCAPEPATCH_API UTextureRenderTarget2D* GetHeightAlphaRenderTarget();
+	LANDSCAPEPATCH_API UTexture2D* GetHeightAlphaInternalTexture();
+
 	LANDSCAPEPATCH_API UTextureRenderTarget2D* GetWeightPatchRenderTarget(ULandscapeWeightPatchTextureInfo* WeightPatch);
 	LANDSCAPEPATCH_API UTexture2D* GetWeightPatchInternalTexture(ULandscapeWeightPatchTextureInfo* WeightPatch);
+	/** Returns the default weight patch alpha render target or the specific weight patch override render target */
+	LANDSCAPEPATCH_API UTextureRenderTarget2D* GetWeightPatchAlphaRenderTarget(ULandscapeWeightPatchTextureInfo* WeightPatch);
+	/** Returns the default weight patch alpha internal texture or the specific weight patch override internal texture */
+	LANDSCAPEPATCH_API UTexture2D* GetWeightPatchAlphaInternalTexture(ULandscapeWeightPatchTextureInfo* WeightPatch);
 
 private:
 	void UpdateHeightConvertToNativeParamsIfNeeded();
+
+	 /**
+	 * If the weight patch overrides the alpha, returns the specific weight patch override settings. Otherwise, returns the default weight patch alpha settings 
+	 */
+	FLandscapeTexturePatchAlphaSettings& GetWeightPatchAlphaSettings(ULandscapeWeightPatchTextureInfo* WeightPatch);
+	const FLandscapeTexturePatchAlphaSettings& GetWeightPatchAlphaSettingsConst(ULandscapeWeightPatchTextureInfo* WeightPatch) const; 
+
 #if WITH_EDITOR
 	void TransitionHeightSourceModeInternal(ELandscapeTexturePatchSourceMode OldMode, ELandscapeTexturePatchSourceMode NewMode);
 	FLandscapeHeightPatchConvertToNativeParams GetHeightConvertToNativeParams() const;
diff --git a/Engine/Source/Runtime/Core/Public/UObject/FortniteMainBranchObjectVersions.inl b/Engine/Source/Runtime/Core/Public/UObject/FortniteMainBranchObjectVersions.inl
--- a/Engine/Source/Runtime/Core/Public/UObject/FortniteMainBranchObjectVersions.inl
+++ b/Engine/Source/Runtime/Core/Public/UObject/FortniteMainBranchObjectVersions.inl
@@ -724,5 +724,8 @@ FORTNITE_MAIN_VERSION(DataflowAddGraphTopologyGuid, 236)
 // Add material usage flag for mesh deformers
 FORTNITE_MAIN_VERSION(AddMeshDeformerUsageFlag, 237)
 
+// Add landscape texture patch alpha override for height/weight patches
+FORTNITE_MAIN_VERSION(LandscapeTexturePatchAlphaOverride, 238)
+
 // -----<new versions can be added above this line>-------------------------------------------------
 #undef FORTNITE_MAIN_VERSION
