Landscape Patch - Added custom texture patch texure transforms for heightmap/weightmap tiling
* Adds scale (relative/world), rotation, and offset properties to the texture height/weight/alpha patch textures

#Synced-CL 49206852

diff --git a/Engine/Plugins/Editor/LandscapePatch/Shaders/Private/LandscapeTexturePatchPS.usf b/Engine/Plugins/Editor/LandscapePatch/Shaders/Private/LandscapeTexturePatchPS.usf
--- a/Engine/Plugins/Editor/LandscapePatch/Shaders/Private/LandscapeTexturePatchPS.usf
+++ b/Engine/Plugins/Editor/LandscapePatch/Shaders/Private/LandscapeTexturePatchPS.usf
@@ -79,15 +79,10 @@ float GetFalloffAlpha(float FalloffWorldMargin, float2 PatchWorldDimensions, flo
 	return Alpha;
 }
 
-float4 SamplePatchTexture(Texture2D<float4> InPatchTexture, SamplerState InPatchTextureSampler, float2 InPatchUVCoordinates)
+// Applies UV space transform (scale, rotation, offset) to a given UV
+float2 GetTransformedUV(float2 InUV, float4x4 InTransform)
 {
-	float4 TextureSample = 0.0f;
-	// We could have used a AM_Border sampler with a border color of 0, but this is not supported on all plaforms : 
-	if (all(InPatchUVCoordinates >= 0.0f) && all(InPatchUVCoordinates <= 1.0f))
-	{
-		TextureSample = InPatchTexture.Sample(InPatchTextureSampler, InPatchUVCoordinates);
-	}
-	return TextureSample;
+	return mul(float4(InUV.xy, 0.0, 1.0), InTransform).xy;	
 }
 
 float GetChannelFromSample(uint InTextureChannel, float4 InSampledValue)
@@ -131,6 +126,7 @@ float GetAlphaFactor(uint InAlphaMode, uint InAlphaTextureChannel, float4 InPatc
 Texture2D<float4> InHeightPatch;
 SamplerState InHeightPatchSampler;
 float4x4 InHeightmapToPatch;
+float4x4 InHeightmapTextureTransformToUV;
 float2 InPatchWorldDimensions;
 float2 InEdgeUVDeadBorder;
 float InFalloffWorldMargin;
@@ -146,10 +142,11 @@ uint InBlendMode;
 uint InFlags;
 // Height source texture color channel
 uint InTextureChannel;
-// Alpha source mode, override texture (White system texture if undefined), and alpha texture color channel
+// Alpha source mode, override texture (White system texture if undefined), alpha texture color channel, and texture transform
 uint InAlphaMode;
 Texture2D<float4> InHeightPatchAlphaTexture;
 uint InAlphaTextureChannel;
+float4x4 InAlphaTextureTransformToUV;
 
 uint ConvertPatchBlendModeToHeightAlphaFlags()
 {
@@ -180,9 +177,13 @@ void ApplyLandscapeTextureHeightPatch(in float4 SVPos : SV_POSITION, out float4
 	float2x2 HeightmapToPatchRotateScale = (float2x2) InHeightmapToPatch;
 	float2 HeightmapToPatchTranslate = InHeightmapToPatch._m03_m13;
 	
+	// Store normalized coordinates that map the patch to a [0, 1] UV space with no tiling (used for patch falloff)
 	float2 PatchUVCoordinates = mul(HeightmapToPatchRotateScale, SVPos.xy) + HeightmapToPatchTranslate;
-	float4 PatchSampledValue = SamplePatchTexture(InHeightPatch, InHeightPatchSampler, PatchUVCoordinates);
+	// Apply the scale, offset, rotation UV space transform to convert from normalized UV to tiling space (utilizes AM_Wrap sampler) 
+	float2 PatchUVCoordinatesTransformed = GetTransformedUV(PatchUVCoordinates, InHeightmapTextureTransformToUV);	
+	float2 PatchAlphaUVCoordinatesTransformed = GetTransformedUV(PatchUVCoordinates, InAlphaTextureTransformToUV);	
 
+	float4 PatchSampledValue = InHeightPatch.Sample(InHeightPatchSampler, PatchUVCoordinatesTransformed);
 	// When using an internal texture/native packed height we enforce that the data comes from the RG channels
 	// Otherwise, get the selected texture channel to use for the height
 	float PatchStoredHeight = bInputIsPackedHeight ? UnpackHeight(PatchSampledValue.xy) : GetChannelFromSample(InTextureChannel, PatchSampledValue);
@@ -191,8 +192,7 @@ void ApplyLandscapeTextureHeightPatch(in float4 SVPos : SV_POSITION, out float4
 	
 	float Alpha = GetFalloffAlpha(InFalloffWorldMargin, InPatchWorldDimensions, PatchUVCoordinates, InEdgeUVDeadBorder, bRectangularFalloff);
 	
-	// TODO [jared.ritchie] Apply alpha specific UV scale/offset/rotation once tiling is added
-	float4 PatchAlphaSampledValue = SamplePatchTexture(InHeightPatchAlphaTexture, InHeightPatchSampler, PatchUVCoordinates);
+	float4 PatchAlphaSampledValue = InHeightPatchAlphaTexture.Sample(InHeightPatchSampler, PatchAlphaUVCoordinatesTransformed);
 	// Apply the alpha from the source height texture or use the standalone alpha texture
 	Alpha *= GetAlphaFactor(InAlphaMode, InAlphaTextureChannel, PatchSampledValue, PatchAlphaSampledValue);
 
@@ -252,6 +252,7 @@ uint2 InSourceWeightmapCoordOffset;
 Texture2D<float4> InWeightPatch;
 SamplerState InWeightPatchSampler;
 float4x4 InWeightmapToPatch;
+float4x4 InWeightmapTextureTransformToUV;
 float2 InPatchWorldDimensions;
 float2 InEdgeUVDeadBorder;
 float InFalloffWorldMargin;
@@ -261,10 +262,11 @@ uint InBlendMode;
 uint InFlags;
 // Weight source texture color channel
 uint InTextureChannel;
-// Alpha source mode, override texture (White system texture if undefined), and alpha texture color channel
+// Alpha source mode, override texture (White system texture if undefined), alpha texture color channel, and texture transform
 uint InAlphaMode;
 Texture2D<float4> InWeightPatchAlphaTexture;
 uint InAlphaTextureChannel;
+float4x4 InAlphaTextureTransformToUV;
 
 uint ConvertPatchBlendModeToWeightAlphaFlags()
 {
@@ -294,16 +296,20 @@ void ApplyLandscapeTextureWeightPatch(in float4 SVPos : SV_POSITION, out float4
 	float2x2 WeightmapToPatchRotateScale = (float2x2) InWeightmapToPatch;
 	float2 WeightmapToPatchTranslate = InWeightmapToPatch._m03_m13;
 	
+	// Store normalized coordinates that map the patch to a [0, 1] UV space with no tiling (used for patch falloff)
 	float2 PatchUVCoordinates = mul(WeightmapToPatchRotateScale, SVPos.xy) + WeightmapToPatchTranslate;
-	float4 PatchSampledValue = SamplePatchTexture(InWeightPatch, InWeightPatchSampler, PatchUVCoordinates);
+	// Apply the scale, offset, rotation UV space transform to convert from normalized UV to tiling space (utilizes AM_Wrap sampler) 
+	float2 PatchUVCoordinatesTransformed = GetTransformedUV(PatchUVCoordinates, InWeightmapTextureTransformToUV);
+	float2 PatchAlphaUVCoordinatesTransformed = GetTransformedUV(PatchUVCoordinates, InAlphaTextureTransformToUV);
+
+	float4 PatchSampledValue = InWeightPatch.Sample(InWeightPatchSampler, PatchUVCoordinatesTransformed);
 	// Get the selected texture color channel from the sample
 	float PatchWeight = GetChannelFromSample(InTextureChannel, PatchSampledValue);
 
 	float Alpha = GetFalloffAlpha(InFalloffWorldMargin, InPatchWorldDimensions, PatchUVCoordinates, InEdgeUVDeadBorder, bRectangularFalloff);
 
-	// TODO [jared.ritchie] Apply alpha specific UV scale/offset/rotation once tiling is added
 	// Apply the alpha from the source weight texture or use the standalone alpha texture
-	float4 PatchAlphaSampledValue = SamplePatchTexture(InWeightPatchAlphaTexture, InWeightPatchSampler, PatchUVCoordinates);
+	float4 PatchAlphaSampledValue = InWeightPatchAlphaTexture.Sample(InWeightPatchSampler, PatchAlphaUVCoordinatesTransformed);
 	Alpha *= GetAlphaFactor(InAlphaMode, InAlphaTextureChannel, PatchSampledValue, PatchAlphaSampledValue);
 	
 	OutColor = float4(PackWeight(PatchWeight), PackWeightAlpha(Alpha, ConvertPatchBlendModeToWeightAlphaFlags()));
diff --git a/Engine/Plugins/Editor/LandscapePatch/Source/LandscapePatch/Private/LandscapeTexturePatch.cpp b/Engine/Plugins/Editor/LandscapePatch/Source/LandscapePatch/Private/LandscapeTexturePatch.cpp
--- a/Engine/Plugins/Editor/LandscapePatch/Source/LandscapePatch/Private/LandscapeTexturePatch.cpp
+++ b/Engine/Plugins/Editor/LandscapePatch/Source/LandscapePatch/Private/LandscapeTexturePatch.cpp
@@ -195,6 +195,22 @@ namespace LandscapeTexturePatchLocals
 	}
 }
 
+FMatrix44d FLandscapeTexturePatchTextureTransform::GetTransformToUV(const FTransform& PatchToWorldIn, const FVector2f& PatchWorldDimensionsIn) const
+{
+	// Calculate UV transforms: Scale (relative or world), Offset, and Rotation. Shader applies transformation to patch UV before sampling
+	// Scale is relative to local UV space unless bUseWorldSpaceScale is enabled. PatchWorldDimensions = PatchUnscaledSize * PatchScale
+	FVector UVOffset = FVector(Offset.X, Offset.Y, 0.0);
+	
+	// Clamp scale to avoid divide by zero
+	const double ScaleX = FMath::Abs(Scale.X) < SMALL_NUMBER ? FMath::FloatSelect(Scale.X, 1.0, -1.0) * SMALL_NUMBER : Scale.X;
+	const double ScaleY = FMath::Abs(Scale.Y) < SMALL_NUMBER ? FMath::FloatSelect(Scale.Y, 1.0, -1.0) * SMALL_NUMBER : Scale.Y;
+	FVector UVRelativeScale = FVector(1.0 / ScaleX, 1.0 / ScaleY, 1.0);
+	const FVector PatchWorldScale = PatchToWorldIn.GetScale3D();
+	FVector2f UVWorldScale = (PatchWorldDimensionsIn * FVector2f(PatchWorldScale.X, PatchWorldScale.Y)) / FVector2f(ScaleX, ScaleY);
+
+	return FTransform(FRotator(0.0, Rotation, 0.0), UVOffset, bUseWorldSpaceScale ? FVector(UVWorldScale.X, UVWorldScale.Y, 1.0) : UVRelativeScale).ToMatrixWithScale();
+}
+
 UTextureRenderTarget2D* FLandscapeTexturePatchAlphaSettings::GetAlphaRenderTarget(UObject* InOuter)
 {
 #if WITH_EDITOR
@@ -551,7 +567,7 @@ UTextureRenderTarget2D* ULandscapeTexturePatch::ApplyToHeightmap(UE::Landscape::
 			FRDGTextureRef PatchTexture = GraphBuilder.RegisterExternalTexture(PatchRenderTarget);
 			FRDGTextureSRVRef PatchSRV = GraphBuilder.CreateSRV(FRDGTextureSRVDesc::CreateForMipLevel(PatchTexture, 0));
 			ShaderParams->InHeightPatch = PatchSRV;
-			ShaderParams->InHeightPatchSampler = TStaticSamplerState<SF_Bilinear, AM_Clamp, AM_Clamp>::GetRHI();
+			ShaderParams->InHeightPatchSampler = TStaticSamplerState<SF_Bilinear, AM_Wrap, AM_Wrap, AM_Wrap>::GetRHI();
 
 			// Heightmap alpha override resource. Pass white system texture when invalid (no transparency)
 			if (PatchAlphaTextureResource != nullptr && PatchAlphaTextureResource->GetTexture2DRHI())
@@ -680,7 +696,7 @@ void ULandscapeTexturePatch::ApplyToWeightmap(UE::Landscape::EditLayers::FRender
 			FRDGTextureRef PatchTexture = GraphBuilder.RegisterExternalTexture(PatchRenderTarget);
 			FRDGTextureSRVRef PatchSRV = GraphBuilder.CreateSRV(FRDGTextureSRVDesc::CreateForMipLevel(PatchTexture, 0));
 			ShaderParams->InWeightPatch = PatchSRV;
-			ShaderParams->InWeightPatchSampler = TStaticSamplerState<SF_Bilinear, AM_Clamp, AM_Clamp>::GetRHI();
+			ShaderParams->InWeightPatchSampler = TStaticSamplerState<SF_Bilinear, AM_Wrap, AM_Wrap, AM_Wrap>::GetRHI();
 
 			// Weightmap alpha override resource. Pass white system texture when invalid (no transparency)
 			if (PatchAlphaTextureResource != nullptr && PatchAlphaTextureResource->GetTexture2DRHI())
@@ -708,8 +724,8 @@ void ULandscapeTexturePatch::ApplyToWeightmap(UE::Landscape::EditLayers::FRender
 }
 
 void ULandscapeTexturePatch::GetCommonShaderParams(const FTransform& LandscapeHeightmapToWorldIn,
-	const FIntPoint& SourceResolutionIn, const FIntPoint& DestinationResolutionIn, 
-	FTransform& PatchToWorldOut, FVector2f& PatchWorldDimensionsOut, FMatrix44f& HeightmapToPatchOut, FIntRect& DestinationBoundsOut, 
+	const FIntPoint& SourceResolutionIn, const FIntPoint& DestinationResolutionIn, const FLandscapeTexturePatchTextureTransform& SourceTextureUVTransformIn,
+	FTransform& PatchToWorldOut, FVector2f& PatchWorldDimensionsOut, FMatrix44f& HeightmapToPatchOut, FMatrix44f& SourceTextureTransformToUVOut, FIntRect& DestinationBoundsOut,
 	FVector2f& EdgeUVDeadBorderOut, float& FalloffWorldMarginOut) const
 {
 	using namespace LandscapeTexturePatchLocals;
@@ -731,9 +747,10 @@ void ULandscapeTexturePatch::GetCommonShaderParams(const FTransform& LandscapeHe
 	// In unreal, matrix composition is done by multiplying the subsequent ones on the right, and the result
 	// is transpose of what our shader will expect (because unreal right multiplies vectors by matrices).
 	FMatrix44d LandscapeToPatchUVTransposed = LandscapeToWorld * WorldToPatch * PatchLocalToUVs;
-	HeightmapToPatchOut = (FMatrix44f)LandscapeToPatchUVTransposed.GetTransposed();
-
+	HeightmapToPatchOut = static_cast<FMatrix44f>(LandscapeToPatchUVTransposed.GetTransposed());
 
+	SourceTextureTransformToUVOut = static_cast<FMatrix44f>(SourceTextureUVTransformIn.GetTransformToUV(PatchToWorldOut, PatchWorldDimensionsOut));
+	
 	// Get the output bounds, which are used to limit the amount of landscape pixels we have to process. 
 	// To get them, convert all of the corners into heightmap 2d coordinates and get the bounding box.
 	auto PatchUVToHeightmap2DCoordinates = [&PatchToWorldOut, &FromPatchUVToPatch, &LandscapeHeightmapToWorldIn](const FVector2f& UV)
@@ -774,8 +791,8 @@ void ULandscapeTexturePatch::GetHeightShaderParams(const FTransform& LandscapeHe
 	using namespace UE::Landscape;
 
 	FTransform PatchToWorld;
-	GetCommonShaderParams(LandscapeHeightmapToWorldIn, SourceResolutionIn, DestinationResolutionIn,
-		PatchToWorld, ParamsOut.InPatchWorldDimensions, ParamsOut.InHeightmapToPatch, 
+	GetCommonShaderParams(LandscapeHeightmapToWorldIn, SourceResolutionIn, DestinationResolutionIn, HeightTextureUVTransform,
+		PatchToWorld, ParamsOut.InPatchWorldDimensions, ParamsOut.InHeightmapToPatch, ParamsOut.InHeightmapTextureTransformToUV,
 		DestinationBoundsOut, ParamsOut.InEdgeUVDeadBorder, ParamsOut.InFalloffWorldMargin);
 
 	FVector3d ComponentScale = PatchToWorld.GetScale3D();
@@ -835,6 +852,7 @@ void ULandscapeTexturePatch::GetHeightShaderParams(const FTransform& LandscapeHe
 
 	ParamsOut.InAlphaMode = static_cast<uint8>(HeightAlphaSettings.GetAlphaSourceMode());
 	ParamsOut.InAlphaTextureChannel = static_cast<uint8>(HeightAlphaSettings.AlphaTextureChannel);
+	ParamsOut.InAlphaTextureTransformToUV = static_cast<FMatrix44f>(HeightAlphaSettings.AlphaTextureUVTransform.GetTransformToUV(PatchToWorld, ParamsOut.InPatchWorldDimensions));
 }
 
 void ULandscapeTexturePatch::GetWeightShaderParams(
@@ -846,8 +864,8 @@ void ULandscapeTexturePatch::GetWeightShaderParams(
 	using namespace UE::Landscape;
 
 	FTransform PatchToWorld;
-	GetCommonShaderParams(LandscapeHeightmapToWorldIn, SourceResolutionIn, DestinationResolutionIn,
-		PatchToWorld, ParamsOut.InPatchWorldDimensions, ParamsOut.InWeightmapToPatch,
+	GetCommonShaderParams(LandscapeHeightmapToWorldIn, SourceResolutionIn, DestinationResolutionIn, WeightPatchInfo->WeightTextureUVTransform,
+		PatchToWorld, ParamsOut.InPatchWorldDimensions, ParamsOut.InWeightmapToPatch, ParamsOut.InWeightmapTextureTransformToUV,
 		DestinationBoundsOut, ParamsOut.InEdgeUVDeadBorder, ParamsOut.InFalloffWorldMargin);
 
 	// Use the override blend mode if present, otherwise fall back to more general blend mode.
@@ -865,6 +883,7 @@ void ULandscapeTexturePatch::GetWeightShaderParams(
 
 	ParamsOut.InAlphaMode = static_cast<uint8>(GetWeightPatchAlphaSourceMode(WeightPatchInfo->WeightmapLayerName));
 	ParamsOut.InAlphaTextureChannel = static_cast<uint8>(GetWeightPatchAlphaTextureChannel(WeightPatchInfo->WeightmapLayerName));
+	ParamsOut.InAlphaTextureTransformToUV = static_cast<FMatrix44f>(GetWeightPatchAlphaTextureTransform(WeightPatchInfo->WeightmapLayerName).GetTransformToUV(PatchToWorld, ParamsOut.InPatchWorldDimensions));
 }
 
 // This function determines how our internal height render targets get converted to the format that gets
@@ -1104,7 +1123,7 @@ void ULandscapeTexturePatch::ReinitializeHeight(UTextureRenderTarget2D* InCombin
 		FRDGTextureRef HeightmapSource = GraphBuilder.RegisterExternalTexture(CreateRenderTarget(Source->GetTexture2DRHI(), TEXT("ReinitializationSource")));
 		FRDGTextureSRVRef SourceSRV = GraphBuilder.CreateSRV(FRDGTextureSRVDesc::CreateForMipLevel(HeightmapSource, 0));
 		HeightmapResampleParams->InSource = SourceSRV;
-		HeightmapResampleParams->InSourceSampler = TStaticSamplerState<SF_Bilinear, AM_Clamp, AM_Clamp>::GetRHI();
+		HeightmapResampleParams->InSourceSampler = TStaticSamplerState<SF_Bilinear, AM_Wrap, AM_Wrap, AM_Wrap>::GetRHI();
 		HeightmapResampleParams->InPatchToSource = PatchToSource;
 
 		FRDGTextureRef DestinationTexture = GraphBuilder.RegisterExternalTexture(CreateRenderTarget(Destination->GetTexture2DRHI(), TEXT("ReinitializationDestination")));
@@ -1247,7 +1266,7 @@ void ULandscapeTexturePatch::ReinitializeWeightPatch(ULandscapeWeightPatchTextur
 			ShaderParams->InSource = GraphBuilder.CreateSRV(Desc);
 		}
 
-		ShaderParams->InSourceSampler = TStaticSamplerState<SF_Bilinear, AM_Clamp, AM_Clamp>::GetRHI();
+		ShaderParams->InSourceSampler = TStaticSamplerState<SF_Bilinear, AM_Wrap, AM_Wrap, AM_Wrap>::GetRHI();
 
 		ShaderParams->InPatchToSource = PatchToSource;
 
@@ -1981,6 +2000,32 @@ void ULandscapeWeightPatchTextureInfo::SetWeightPatchAlphaTextureChannel(ELandsc
 #endif // WITH_EDITOR
 }
 
+void ULandscapeWeightPatchTextureInfo::SetWeightPatchTextureTransform(const FLandscapeTexturePatchTextureTransform& InTextureTransform, bool bAlwaysMarkDirty)
+{
+#if WITH_EDITOR
+	if (WeightTextureUVTransform == InTextureTransform)
+	{
+		return;
+	}
+	Modify(bAlwaysMarkDirty);
+
+	WeightTextureUVTransform = InTextureTransform;
+#endif // WITH_EDITOR
+}
+
+void ULandscapeWeightPatchTextureInfo::SetWeightPatchAlphaTextureTransform(const FLandscapeTexturePatchTextureTransform& InTextureTransform, bool bAlwaysMarkDirty)
+{
+#if WITH_EDITOR
+	if (WeightPatchAlphaSettings.AlphaTextureUVTransform == InTextureTransform)
+	{
+		return;
+	}
+	Modify(bAlwaysMarkDirty);
+
+	WeightPatchAlphaSettings.AlphaTextureUVTransform = InTextureTransform;
+#endif // WITH_EDITOR
+}
+
 #if WITH_EDITOR
 void ULandscapeWeightPatchTextureInfo::TransitionSourceModeInternal(ELandscapeTexturePatchSourceMode OldMode, ELandscapeTexturePatchSourceMode NewMode)
 {
@@ -2029,6 +2074,19 @@ void ULandscapeTexturePatch::SetHeightSourceMode(ELandscapeTexturePatchSourceMod
 #endif // WITH_EDITOR
 }
 
+void ULandscapeTexturePatch::SetHeightTextureTransform(const FLandscapeTexturePatchTextureTransform& InTextureTransform, bool bAlwaysMarkDirty)
+{
+#if WITH_EDITOR
+	if (HeightTextureUVTransform == InTextureTransform)
+	{
+		return;
+	}
+	Modify(bAlwaysMarkDirty);
+
+	HeightTextureUVTransform = InTextureTransform;
+#endif // WITH_EDITOR
+}
+
 void ULandscapeTexturePatch::SetHeightTextureChannel(ELandscapeTexturePatchTextureChannel InTextureChannel, bool bAlwaysMarkDirty)
 {
 #if WITH_EDITOR
@@ -2057,6 +2115,19 @@ void ULandscapeTexturePatch::SetHeightAlphaSourceMode(ELandscapeTexturePatchAlph
 #endif // WITH_EDITOR
 }
 
+void ULandscapeTexturePatch::SetHeightAlphaTextureTransform(const FLandscapeTexturePatchTextureTransform& InTextureTransform, bool bAlwaysMarkDirty)
+{
+#if WITH_EDITOR
+	if (HeightAlphaSettings.AlphaTextureUVTransform == InTextureTransform)
+	{
+		return;
+	}
+	Modify(bAlwaysMarkDirty);
+
+	HeightAlphaSettings.AlphaTextureUVTransform = InTextureTransform;
+#endif // WITH_EDITOR
+}
+
 void ULandscapeTexturePatch::SetHeightAlphaTextureChannel(ELandscapeTexturePatchTextureChannel InTextureChannel, bool bAlwaysMarkDirty)
 {
 #if WITH_EDITOR
@@ -2452,6 +2523,10 @@ void ULandscapeTexturePatch::SetWeightPatchAlphaTextureAsset(const FName& InWeig
 	// Set the default weightmap alpha texture for all weightmap patches
 	if (InWeightmapLayerName.IsNone())
 	{
+		if (WeightPatchesAlphaSettings.AlphaTextureAsset == InTexture)
+		{
+			return;
+		}
 		Modify(bAlwaysMarkDirty);
 		WeightPatchesAlphaSettings.AlphaTextureAsset = InTexture;
 		return;
@@ -2479,6 +2554,52 @@ void ULandscapeTexturePatch::SetWeightPatchAlphaTextureAsset(const FName& InWeig
 	UE_LOG(LogLandscapePatch, Warning, TEXT("ULandscapeTexturePatch::SetWeightPatchAlphaTextureAsset: Unable to find data for weight layer %s"), *LayerNameString);
 }
 
+const FLandscapeTexturePatchTextureTransform& ULandscapeTexturePatch::GetWeightPatchAlphaTextureTransform(const FName& InWeightmapLayerName) const
+{
+	// Get the default weightmap alpha texture transform for all weightmap patches
+	if (InWeightmapLayerName.IsNone())
+	{
+		return WeightPatchesAlphaSettings.AlphaTextureUVTransform;
+	}
+
+	for (const TObjectPtr<ULandscapeWeightPatchTextureInfo>& WeightPatch : WeightPatches)
+	{
+		if (IsValid(WeightPatch) && WeightPatch->WeightmapLayerName == InWeightmapLayerName)
+		{
+			return GetWeightPatchAlphaSettingsConst(WeightPatch).AlphaTextureUVTransform;
+		}
+	}
+
+	UE_LOG(LogLandscapePatch, Warning, TEXT("ULandscapeTexturePatch::GetWeightPatchAlphaTextureTransform: Unable to get data for weight layer %s"), *InWeightmapLayerName.ToString());
+	return WeightPatchesAlphaSettings.AlphaTextureUVTransform;
+}
+
+void ULandscapeTexturePatch::SetWeightPatchAlphaTextureTransform(const FName& InWeightmapLayerName, const FLandscapeTexturePatchTextureTransform& InTextureTransform, bool bAlwaysMarkDirty)
+{
+	// Set the default weightmap alpha texture transform for all weightmap patches
+	if (InWeightmapLayerName.IsNone())
+	{
+		if (WeightPatchesAlphaSettings.AlphaTextureUVTransform == InTextureTransform)
+		{
+			return;
+		}
+		Modify(bAlwaysMarkDirty);
+		WeightPatchesAlphaSettings.AlphaTextureUVTransform = InTextureTransform;
+		return;
+	}
+
+	for (const TObjectPtr<ULandscapeWeightPatchTextureInfo>& WeightPatch : WeightPatches)
+	{
+		if (IsValid(WeightPatch) && WeightPatch->WeightmapLayerName == InWeightmapLayerName)
+		{
+			WeightPatch->SetWeightPatchAlphaTextureTransform(InTextureTransform, bAlwaysMarkDirty);
+			return;
+		}
+	}
+
+	UE_LOG(LogLandscapePatch, Warning, TEXT("ULandscapeTexturePatch::SetWeightPatchAlphaTextureTransform: Unable to set data for weight layer %s"), *InWeightmapLayerName.ToString());
+}
+
 ELandscapeTexturePatchTextureChannel ULandscapeTexturePatch::GetWeightPatchAlphaTextureChannel(const FName& InWeightmapLayerName) const
 {
 	if (InWeightmapLayerName.IsNone())
@@ -2564,6 +2685,35 @@ ELandscapeTexturePatchSourceMode ULandscapeTexturePatch::GetWeightPatchSourceMod
 	return ELandscapeTexturePatchSourceMode::None;
 }
 
+const FLandscapeTexturePatchTextureTransform& ULandscapeTexturePatch::GetWeightPatchTextureTransform(const FName& InWeightmapLayerName) const
+{
+	for (const TObjectPtr<ULandscapeWeightPatchTextureInfo>& WeightPatch : WeightPatches)
+	{
+		if (IsValid(WeightPatch) && WeightPatch->WeightmapLayerName == InWeightmapLayerName)
+		{
+			return WeightPatch->WeightTextureUVTransform;
+		}
+	}
+
+	// There is no default texture transform for all weight patches (Only the alpha texture is overridable). Log an explicit error here
+	UE_LOG(LogLandscapePatch, Error, TEXT("ULandscapeTexturePatch::GetWeightPatchTextureTransform: Unable to get transform data for weight layer %s. The weight patch transform will fallback to the height texture patch transform."), *InWeightmapLayerName.ToString());
+	return HeightTextureUVTransform;
+}
+
+void ULandscapeTexturePatch::SetWeightPatchTextureTransform(const FName& InWeightmapLayerName, const FLandscapeTexturePatchTextureTransform& InTextureTransform, bool bAlwaysMarkDirty)
+{
+	for (const TObjectPtr<ULandscapeWeightPatchTextureInfo>& WeightPatch : WeightPatches)
+	{
+		if (IsValid(WeightPatch) && WeightPatch->WeightmapLayerName == InWeightmapLayerName)
+		{
+			WeightPatch->SetWeightPatchTextureTransform(InTextureTransform, bAlwaysMarkDirty);
+			return;
+		}
+	}
+
+	UE_LOG(LogLandscapePatch, Warning, TEXT("ULandscapeTexturePatch::SetWeightPatchTextureTransform: Unable to set data for weight layer %s"), *InWeightmapLayerName.ToString());
+}
+
 ELandscapeTexturePatchTextureChannel ULandscapeTexturePatch::GetWeightPatchTextureChannel(const FName& InWeightmapLayerName) const
 {
 	for (const TObjectPtr<ULandscapeWeightPatchTextureInfo>& WeightPatch : WeightPatches)
diff --git a/Engine/Plugins/Editor/LandscapePatch/Source/LandscapePatch/Private/LandscapeTexturePatchPS.cpp b/Engine/Plugins/Editor/LandscapePatch/Source/LandscapePatch/Private/LandscapeTexturePatchPS.cpp
--- a/Engine/Plugins/Editor/LandscapePatch/Source/LandscapePatch/Private/LandscapeTexturePatchPS.cpp
+++ b/Engine/Plugins/Editor/LandscapePatch/Source/LandscapePatch/Private/LandscapeTexturePatchPS.cpp
@@ -105,7 +105,7 @@ void FSimpleTextureCopyPS::AddToRenderGraph(FRDGBuilder& GraphBuilder, FRDGTextu
 
 	FSimpleTextureCopyPS::FParameters* ShaderParams = GraphBuilder.AllocParameters<FSimpleTextureCopyPS::FParameters>();
 	ShaderParams->InSource = GraphBuilder.CreateSRV(FRDGTextureSRVDesc::CreateForMipLevel(SourceTexture, 0));
-	ShaderParams->InSourceSampler = TStaticSamplerState<SF_Bilinear, AM_Clamp, AM_Clamp>::GetRHI();
+	ShaderParams->InSourceSampler = TStaticSamplerState<SF_Bilinear, AM_Wrap, AM_Wrap, AM_Wrap>::GetRHI();
 	ShaderParams->RenderTargets[0] = FRenderTargetBinding(DestinationTexture, ERenderTargetLoadAction::ENoAction, /*InMipIndex = */0);
 	FIntVector DestinationSize = DestinationTexture->Desc.GetSize();
 	ShaderParams->InDestinationResolution = FVector2f(DestinationSize.X, DestinationSize.Y);
diff --git a/Engine/Plugins/Editor/LandscapePatch/Source/LandscapePatch/Private/LandscapeTexturePatchPS.h b/Engine/Plugins/Editor/LandscapePatch/Source/LandscapePatch/Private/LandscapeTexturePatchPS.h
--- a/Engine/Plugins/Editor/LandscapePatch/Source/LandscapePatch/Private/LandscapeTexturePatchPS.h
+++ b/Engine/Plugins/Editor/LandscapePatch/Source/LandscapePatch/Private/LandscapeTexturePatchPS.h
@@ -45,6 +45,7 @@ namespace UE::Landscape
 			SHADER_PARAMETER_RDG_TEXTURE_SRV(Texture2D<float4>, InHeightPatch)
 			SHADER_PARAMETER_SAMPLER(SamplerState, InHeightPatchSampler)
 			SHADER_PARAMETER(FMatrix44f, InHeightmapToPatch)
+			SHADER_PARAMETER(FMatrix44f, InHeightmapTextureTransformToUV)
 			// Value in patch that corresponds to the landscape mid value, which is our "0 height".
 			SHADER_PARAMETER(float, InZeroInEncoding)
 			// Scale to apply to source values relative to the value that represents 0 height.
@@ -63,10 +64,11 @@ namespace UE::Landscape
 			SHADER_PARAMETER(uint32, InFlags)
 			// Height source texture color channel
 			SHADER_PARAMETER(uint32, InTextureChannel)
-			// Alpha source mode, override texture (White system texture if undefined), and texture alpha color channel
+			// Alpha source mode, override texture (White system texture if undefined), alpha texture color channel, and texture transform
 			SHADER_PARAMETER(uint32, InAlphaMode)
 			SHADER_PARAMETER_RDG_TEXTURE_SRV(Texture2D<float4>, InHeightPatchAlphaTexture)
 			SHADER_PARAMETER(uint32, InAlphaTextureChannel)
+			SHADER_PARAMETER(FMatrix44f, InAlphaTextureTransformToUV)
 
 			RENDER_TARGET_BINDING_SLOTS() // Holds our output
 		END_SHADER_PARAMETER_STRUCT()
@@ -183,6 +185,7 @@ namespace UE::Landscape
 			SHADER_PARAMETER_RDG_TEXTURE_SRV(Texture2D<float4>, InWeightPatch)
 			SHADER_PARAMETER_SAMPLER(SamplerState, InWeightPatchSampler)
 			SHADER_PARAMETER(FMatrix44f, InWeightmapToPatch)
+			SHADER_PARAMETER(FMatrix44f, InWeightmapTextureTransformToUV)
 			// Amount of the patch edge to not apply in UV space. Generally set to 0.5/Dimensions to avoid applying
 			// the edge half-pixels.
 			SHADER_PARAMETER(FVector2f, InEdgeUVDeadBorder)
@@ -195,10 +198,11 @@ namespace UE::Landscape
 			SHADER_PARAMETER(uint32, InFlags)
 			// Weight source texture color channel
 			SHADER_PARAMETER(uint32, InTextureChannel)
-			// Alpha source mode, override texture (White system texture if undefined), and alpha texture color channel
+			// Alpha source mode, override texture (White system texture if undefined), alpha texture color channel, and texture transform
 			SHADER_PARAMETER(uint32, InAlphaMode)
 			SHADER_PARAMETER_RDG_TEXTURE_SRV(Texture2D<float4>, InWeightPatchAlphaTexture)
 			SHADER_PARAMETER(uint32, InAlphaTextureChannel)
+			SHADER_PARAMETER(FMatrix44f, InAlphaTextureTransformToUV)
 
 			RENDER_TARGET_BINDING_SLOTS() // Holds our output
 		END_SHADER_PARAMETER_STRUCT()
diff --git a/Engine/Plugins/Editor/LandscapePatch/Source/LandscapePatch/Public/LandscapeTexturePatch.h b/Engine/Plugins/Editor/LandscapePatch/Source/LandscapePatch/Public/LandscapeTexturePatch.h
--- a/Engine/Plugins/Editor/LandscapePatch/Source/LandscapePatch/Public/LandscapeTexturePatch.h
+++ b/Engine/Plugins/Editor/LandscapePatch/Source/LandscapePatch/Public/LandscapeTexturePatch.h
@@ -182,6 +182,40 @@ enum class ELandscapeTextureHeightPatchZeroHeightMeaning : uint8
 	WorldZero
 };
 
+USTRUCT(BlueprintType)
+struct FLandscapeTexturePatchTextureTransform
+{
+	GENERATED_BODY()
+public:
+	/** If enabled, texture UVs are scaled using world-space units (1000 = tile texture every 10 meters). If disabled, scaling is applied in local UV space. */
+	UPROPERTY(EditAnywhere, Category = WeightPatch)
+	bool bUseWorldSpaceScale = false;
+
+	/** Scale applied to the texture's UV. By default, scaling is applied in local UV space. Relative scale 2.0 = texture is twice as large as original size. 
+	* When World Space Scale is enabled, scale 1000 = tile texture every 10 meters. 
+	*/
+	UPROPERTY(EditAnywhere, Category = WeightPatch)
+	FVector2D Scale = FVector2D(1, 1);
+	
+	/** Rotation in degrees applied to the texture's UV. */
+	UPROPERTY(EditAnywhere, Category = WeightPatch, meta = (UIMin = "0", UIMax = "360", ClampMin = "0", ClampMax = "360"))
+	float Rotation = 0;
+
+	/** Relative offset applied to the texture's UV. */
+	UPROPERTY(EditAnywhere, Category = WeightPatch, meta = (UIMin = "0", UIMax = "1", ClampMin = "0", ClampMax = "1"))
+	FVector2D Offset = FVector2D(0, 0);
+
+	FMatrix44d GetTransformToUV(const FTransform& PatchToWorldIn, const FVector2f& PatchWorldDimensionsIn) const;
+
+	FORCEINLINE bool operator==(const FLandscapeTexturePatchTextureTransform InTransform) const
+	{
+		return bUseWorldSpaceScale == InTransform.bUseWorldSpaceScale && 
+			Scale.X == InTransform.Scale.X && Scale.Y == InTransform.Scale.Y && 
+			Offset.X == InTransform.Offset.X && Offset.Y == InTransform.Offset.Y &&
+			Rotation == InTransform.Rotation;
+	}
+};
+
 USTRUCT(BlueprintType)
 struct FLandscapeTexturePatchAlphaSettings
 {
@@ -221,7 +255,9 @@ public:
 	UPROPERTY(VisibleAnywhere, Category = "Texture Alpha Settings", Instanced, AdvancedDisplay, meta = (DisplayName = "Debug Data", EditConditionHides, EditCondition = "AlphaSourceMode != ELandscapeTexturePatchAlphaSourceMode::None"))
 	TObjectPtr<ULandscapeWeightTextureBackedRenderTarget> InternalAlphaData;
 
-	// TODO [jared.ritchie] Add Scale/Offset/Rotation UV properties here
+	/* Alpha texture's UV transform (Scale, Offset, Rotation) */
+	UPROPERTY(EditAnywhere, Category = "Texture Alpha Settings", meta = (DisplayName = "Alpha Texture Transform"), meta = (EditConditionHides, EditCondition = "AlphaSourceMode != ELandscapeTexturePatchAlphaSourceMode::None"))
+	FLandscapeTexturePatchTextureTransform AlphaTextureUVTransform;
 
 private:
 	/** When the source mode changes, there are transition steps that need run to clean up render targets / internal data.
@@ -275,11 +311,11 @@ public:
 
 protected:
 	/** Name of the target layer to affect. */
-	UPROPERTY(EditAnywhere, Category = WeightPatch)
+	UPROPERTY(EditAnywhere, Category = WeightPatch, meta = (DisplayName = "Target Layer Name"))
 	FName WeightmapLayerName;
 
-	/** Specifies if this patch edits the visibility layer. */
-	UPROPERTY(EditAnywhere, Category = WeightPatch)
+	/** Specifies if this patch affects the visibility layer. */
+	UPROPERTY(EditAnywhere, Category = WeightPatch, meta = (DisplayName = "Affects Visibility Layer"))
 	bool bEditVisibilityLayer = false;
 	
 	UPROPERTY(EditAnywhere, Category = WeightPatch, meta = (EditConditionHides, EditCondition = "false"))
@@ -295,6 +331,9 @@ protected:
 		DisallowedAssetDataTags = "VirtualTextureStreaming=True"))
 	TObjectPtr<UTexture> TextureAsset = nullptr;
 
+	UPROPERTY(EditAnywhere, Category = WeightPatch, meta = (DisplayName = "Texture Transform"), meta = (EditConditionHides, EditCondition = "SourceMode != ELandscapeTexturePatchSourceMode::None"))
+	FLandscapeTexturePatchTextureTransform WeightTextureUVTransform;
+
 	/** Color channel of the texture used for weight data. By default the weight texture's red channel is used. */
 	UPROPERTY(EditAnywhere, Category = HeightPatch, meta = (EditConditionHides, EditCondition = "SourceMode != ELandscapeTexturePatchSourceMode::None"))
 	ELandscapeTexturePatchTextureChannel WeightPatchTextureChannel = ELandscapeTexturePatchTextureChannel::Red;
@@ -330,6 +369,8 @@ protected:
 	void SetWeightPatchAlphaSourceMode(ELandscapeTexturePatchAlphaSourceMode InSourceMode, bool bInOverrideAlpha, bool bAlwaysMarkDirty = true);
 	void SetWeightPatchTextureChannel(ELandscapeTexturePatchTextureChannel InTextureChannel, bool bAlwaysMarkDirty = true);
 	void SetWeightPatchAlphaTextureChannel(ELandscapeTexturePatchTextureChannel InTextureChannel, bool bAlwaysMarkDirty = true);
+	void SetWeightPatchTextureTransform(const FLandscapeTexturePatchTextureTransform& InTextureTransform, bool bAlwaysMarkDirty = true);
+	void SetWeightPatchAlphaTextureTransform(const FLandscapeTexturePatchTextureTransform& InTextureTransform, bool bAlwaysMarkDirty = true);
 
 #if WITH_EDITOR
 	void TransitionSourceModeInternal(ELandscapeTexturePatchSourceMode OldMode, ELandscapeTexturePatchSourceMode NewMode);
@@ -489,6 +530,17 @@ public:
 	UFUNCTION(BlueprintCallable, Category = LandscapePatch)
 	LANDSCAPEPATCH_API virtual void SetHeightSourceMode(ELandscapeTexturePatchSourceMode NewMode);
 
+	UFUNCTION(BlueprintCallable, Category = LandscapePatch)
+	LANDSCAPEPATCH_API virtual const FLandscapeTexturePatchTextureTransform& GetHeightTextureTransform() const { return HeightTextureUVTransform; };
+
+	/**
+	 * Sets the height texture's UV transform. Transform is applied relative UVSpace unless struct property bInWorldSpaceScale is true
+	 * @param InTextureTransform the height texture's transform
+	 * @param bAlwaysMarkDirty mark the package dirty when transform changes
+	 */
+	UFUNCTION(BlueprintCallable, Category = LandscapePatch)
+	LANDSCAPEPATCH_API virtual void SetHeightTextureTransform(const FLandscapeTexturePatchTextureTransform& InTextureTransform, bool bAlwaysMarkDirty = true);
+
 	/**
 	 * Sets the height texture's color channel. Texture channel is restricted when HeightSource is InternalData or HeightEncoding is NativePackedHeight (Assumes RG channels)
 	 * @param InTextureChannel the color channel of the texture to use
@@ -511,6 +563,17 @@ public:
 	UFUNCTION(BlueprintCallable, Category = LandscapePatch)
 	LANDSCAPEPATCH_API virtual void SetHeightAlphaSourceMode(ELandscapeTexturePatchAlphaSourceMode InNewMode, bool bAlwaysMarkDirty = true);
 
+	UFUNCTION(BlueprintCallable, Category = LandscapePatch)
+	LANDSCAPEPATCH_API virtual const FLandscapeTexturePatchTextureTransform& GetHeightAlphaTextureTransform() const { return HeightAlphaSettings.AlphaTextureUVTransform; };
+
+	/**
+	 * Sets the height alpha texture's UV transform. Transform is applied relative UVSpace unless struct property bInWorldSpaceScale is true
+	 * @param InTextureTransform the height alpha texture's transform
+	 * @param bAlwaysMarkDirty mark the package dirty when alpha transform changes
+	 */
+	UFUNCTION(BlueprintCallable, Category = LandscapePatch)
+	LANDSCAPEPATCH_API virtual void SetHeightAlphaTextureTransform(const FLandscapeTexturePatchTextureTransform& InTextureTransform, bool bAlwaysMarkDirty = true);
+
 	UFUNCTION(BlueprintCallable, Category = LandscapePatch)
 	virtual ELandscapeTexturePatchTextureChannel GetHeightAlphaSourceTextureChannel() const { return HeightAlphaSettings.AlphaTextureChannel; }
 
@@ -638,6 +701,18 @@ public:
 	UFUNCTION(BlueprintCallable, Category = LandscapePatch)
 	LANDSCAPEPATCH_API virtual ELandscapeTexturePatchSourceMode GetWeightPatchSourceMode(const FName& InWeightmapLayerName);
 
+	UFUNCTION(BlueprintCallable, Category = LandscapePatch)
+	LANDSCAPEPATCH_API virtual const FLandscapeTexturePatchTextureTransform& GetWeightPatchTextureTransform(const FName& InWeightmapLayerName) const;
+
+	/**
+	 * Sets the weight texture's UV transform. Transform is applied relative UVSpace unless struct property bInWorldSpaceScale is true
+	 * @param InWeightmapLayerName weightmap layer name to set the texture transform
+	 * @param InTextureTransform the weight texture's transform
+	 * @param bAlwaysMarkDirty mark the package dirty when transform changes
+	 */
+	UFUNCTION(BlueprintCallable, Category = LandscapePatch)
+	LANDSCAPEPATCH_API virtual void SetWeightPatchTextureTransform(const FName& InWeightmapLayerName, const FLandscapeTexturePatchTextureTransform& InTextureTransform, bool bAlwaysMarkDirty = true);
+
 	UFUNCTION(BlueprintCallable, Category = LandscapePatch)
 	LANDSCAPEPATCH_API virtual UTexture* GetWeightPatchTextureAsset(const FName& InWeightmapLayerName);
 
@@ -695,6 +770,19 @@ public:
 	UFUNCTION(BlueprintCallable, Category = LandscapePatch)
 	LANDSCAPEPATCH_API void SetWeightPatchAlphaTextureAsset(const FName& InWeightmapLayerName, UTexture* InTexture, bool bInOverrideAlpha = true, bool bAlwaysMarkDirty = true);
 
+	UFUNCTION(BlueprintCallable, Category = LandscapePatch)
+	LANDSCAPEPATCH_API virtual const FLandscapeTexturePatchTextureTransform& GetWeightPatchAlphaTextureTransform(const FName& InWeightmapLayerName) const;
+
+	/**
+	 * Sets the weight texture's alpha UV transform. Transform is applied relative UVSpace unless struct property bInWorldSpaceScale is true
+	 * @param InWeightmapLayerName weightmap layer name to set the texture transform
+	 * @param InTextureTransform the weight alpha texture's transform
+	 * @param bAlwaysMarkDirty mark the package dirty when alpha transform changes
+	 */
+	UFUNCTION(BlueprintCallable, Category = LandscapePatch)
+	LANDSCAPEPATCH_API virtual void SetWeightPatchAlphaTextureTransform(const FName& InWeightmapLayerName, const FLandscapeTexturePatchTextureTransform& InTextureTransform, bool bAlwaysMarkDirty = true);
+
+
 	UFUNCTION(BlueprintCallable, Category = LandscapePatch)
 	LANDSCAPEPATCH_API virtual ELandscapeTexturePatchTextureChannel GetWeightPatchAlphaTextureChannel(const FName& InWeightmapLayerName) const;
 
@@ -781,6 +869,9 @@ protected:
 		DisallowedAssetDataTags = "VirtualTextureStreaming=True"))
 	TObjectPtr<UTexture> HeightTextureAsset = nullptr;
 
+	UPROPERTY(EditAnywhere, Category = HeightPatch, meta = (DisplayName = "Texture Transform"), meta = (EditConditionHides, EditCondition = "HeightSourceMode != ELandscapeTexturePatchSourceMode::None"))
+	FLandscapeTexturePatchTextureTransform HeightTextureUVTransform;
+
 	UE_DEPRECATED(5.8, "bUseAlphaChannel is deprecated, use the ELandscapeTexturePatchAlphaSourceMode to define the alpha behavior.")
 	UPROPERTY(meta = (DeprecatedProperty, DeprecationMessage = "Use HeightAlphaSettings instead. See ELandscapeTexturePatchAlphaSourceMode::SourceTextureChannel to enable legacy behavior"))
 	bool bUseTextureAlphaForHeight_DEPRECATED = false;
@@ -912,9 +1003,9 @@ private:
 		FTextureResource* InMergedLandscapeTextureResource, int32 LandscapeTextureSliceIndex, const FIntPoint& LandscapeTextureResolution, 
 		const FTransform& LandscapeHeightmapToWorld, bool& bHasRenderedSomething, ERHIAccess OutputAccess = ERHIAccess::None);
 
-	void GetCommonShaderParams(const FTransform& LandscapeHeightmapToWorldIn, 
-		const FIntPoint& SourceResolutionIn, const FIntPoint& DestinationResolutionIn,
-		FTransform& PatchToWorldOut, FVector2f& PatchWorldDimensionsOut, FMatrix44f& HeightmapToPatchOut, 
+	void GetCommonShaderParams(const FTransform& LandscapeHeightmapToWorldIn,
+		const FIntPoint& SourceResolutionIn, const FIntPoint& DestinationResolutionIn, const FLandscapeTexturePatchTextureTransform& SourceTextureUVTransformIn,
+		FTransform& PatchToWorldOut, FVector2f& PatchWorldDimensionsOut, FMatrix44f& HeightmapToPatchOut, FMatrix44f& SourceTextureTransformToUVOut,
 		FIntRect& DestinationBoundsOut, FVector2f& EdgeUVDeadBorderOut, float& FalloffWorldMarginOut) const;
 	void GetHeightShaderParams(const FTransform& LandscapeHeightmapToWorldIn, 
 		const FIntPoint& SourceResolutionIn, const FIntPoint& DestinationResolutionIn,
