* Exposed landscape edit layers and their names to BP
* Added BP function to LandscapePatchComponent to assign both the landscape actor and edit layer (by name). This avoids the burden of dealing with edit layer FGuids on the BP side

#Synced-CL 49356504

diff --git a/Engine/Plugins/Editor/LandscapePatch/Source/LandscapePatch/Private/LandscapePatchComponent.cpp b/Engine/Plugins/Editor/LandscapePatch/Source/LandscapePatch/Private/LandscapePatchComponent.cpp
--- a/Engine/Plugins/Editor/LandscapePatch/Source/LandscapePatch/Private/LandscapePatchComponent.cpp
+++ b/Engine/Plugins/Editor/LandscapePatch/Source/LandscapePatch/Private/LandscapePatchComponent.cpp
@@ -844,6 +844,49 @@ void ULandscapePatchComponent::SetEditLayerGuid(const FGuid& GuidIn)
 #endif
 }
 
+bool ULandscapePatchComponent::AssignToLandscape(ALandscape* InLandscape, const FName& InEditLayerName)
+{
+	bool bSuccess = false;
+
+#if WITH_EDITOR
+	Modify(false);
+	bSuccess = BindToLandscape(InLandscape);
+	if (bSuccess && (InLandscape != nullptr))
+	{
+		const ULandscapeEditLayerBase* LocalEditLayer = InLandscape->GetEditLayerConst(InEditLayerName);
+		if (LocalEditLayer != nullptr)
+		{
+			const ULandscapePatchEditLayer* PatchEditLayer = Cast<const ULandscapePatchEditLayer>(LocalEditLayer);
+			if (PatchEditLayer != nullptr)
+			{
+				bSuccess = BindToEditLayer(LocalEditLayer->GetGuid());
+			}
+			else
+			{
+				UE_LOG(LogLandscapePatch, Warning, TEXT("Edit layer with the given name (%s) on landscape %s is not a patch edit layer (%s). "
+					"(Package: %s"),
+					*InEditLayerName.ToString(),
+					*InLandscape->GetActorLabel(),
+					*LocalEditLayer->GetClass()->GetName(),
+					*GetNameSafe(GetPackage()));
+				bSuccess = false;
+			}
+		}
+		else
+		{
+			UE_LOG(LogLandscapePatch, Warning, TEXT("Could not find an edit layer with the given name (%s) on landscape %s for the patch. "
+				"(Package: %s"),
+				*InEditLayerName.ToString(),
+				*InLandscape->GetActorLabel(),
+				*GetNameSafe(GetPackage()));
+			bSuccess = false;
+		}
+	}
+#endif
+
+	return bSuccess;
+}
+
 uint32 ULandscapePatchComponent::GetFullNameHash() const
 {
 	// TODO: If needed, we can cache this value, but then we need to update it in
diff --git a/Engine/Plugins/Editor/LandscapePatch/Source/LandscapePatch/Public/LandscapePatchComponent.h b/Engine/Plugins/Editor/LandscapePatch/Source/LandscapePatch/Public/LandscapePatchComponent.h
--- a/Engine/Plugins/Editor/LandscapePatch/Source/LandscapePatch/Public/LandscapePatchComponent.h
+++ b/Engine/Plugins/Editor/LandscapePatch/Source/LandscapePatch/Public/LandscapePatchComponent.h
@@ -103,6 +103,16 @@ public:
 	UFUNCTION(BlueprintCallable, Category = "LandscapePatch")
 	LANDSCAPEPATCH_API void SetEditLayerGuid(const FGuid& GuidIn);
 
+		
+	/**
+	 * Tries to assign the patch to the landscape 
+	 * @param InLandscape Landscape actor on which the patch component should be applied. Passing an null actor will un-assign the patch component from any landscape
+	 * @param InEditLayerName Name of the edit layer on which the patch component should be assigned. The layer name has to be present on the landscape and it has to be a Landscape Patch Edit Layer for this to succeed
+	 * @return true if the patch was successfully assigned (or un-assigned) to (from) the landscape's edit layer
+	 */
+	UFUNCTION(BlueprintCallable, Category = "LandscapePatch")
+	LANDSCAPEPATCH_API bool AssignToLandscape(ALandscape* InLandscape, const FName& InEditLayerName);
+
 	/**
 	 * Returns the result of GetTypeHash(GetFullName()) for this component. Useful for quick deterministic
 	 * (though not lexical) ordering of same-priority patches. 
diff --git a/Engine/Source/Runtime/Landscape/Classes/Landscape.h b/Engine/Source/Runtime/Landscape/Classes/Landscape.h
--- a/Engine/Source/Runtime/Landscape/Classes/Landscape.h
+++ b/Engine/Source/Runtime/Landscape/Classes/Landscape.h
@@ -382,6 +382,18 @@ public:
 	UFUNCTION(BlueprintCallable, BlueprintPure=false, Category = "Landscape|Editor", meta=(DevelopmentOnly))
 	TArray<FName> GetTargetLayerNames(bool bInIncludeVisibilityLayer = false) const;
 
+	/**
+	* @return the list of edit layers on this landscape (editor-only : returns nothing at runtime)
+	*/
+	UFUNCTION(BlueprintCallable, DisplayName = GetEditLayers, Category = "Landscape|Editor", meta = (DevelopmentOnly))
+	const TArray<ULandscapeEditLayerBase*> GetEditLayersBP() const;
+
+	/**
+	* @return the edit layer whose name matches the parameter on this landscape (editor-only : returns nothing at runtime)
+	*/
+	UFUNCTION(BlueprintCallable, DisplayName = GetEditLayer, Category = "Landscape|Editor", meta = (DevelopmentOnly))
+	const ULandscapeEditLayerBase* GetEditLayerByNameBP(const FName& InLayerName) const;
+
 	bool IsValidRenderTargetFormatHeightmap(EPixelFormat InRenderTargetFormat, bool& bOutCompressHeight);
 	bool IsValidRenderTargetFormatWeightmap(EPixelFormat InRenderTargetFormat, int32& OutNumChannels);
 
diff --git a/Engine/Source/Runtime/Landscape/Classes/LandscapeEditLayer.h b/Engine/Source/Runtime/Landscape/Classes/LandscapeEditLayer.h
--- a/Engine/Source/Runtime/Landscape/Classes/LandscapeEditLayer.h
+++ b/Engine/Source/Runtime/Landscape/Classes/LandscapeEditLayer.h
@@ -38,7 +38,7 @@ struct FOnLandscapeEditLayerDataChangedParams
 * Base class for all landscape edit layers. By implementing the various virtual functions, we are able to customize the behavior of the edit layer
 *  wrt the landscape tools in a generic way (e.g. does it support sculpting tools? painting tools? can it be collapsed?, etc.)
 */
-UCLASS(MinimalAPI, Abstract)
+UCLASS(MinimalAPI, Abstract, BlueprintType)
 class ULandscapeEditLayerBase : public UObject
 #if CPP && WITH_EDITOR
 	, public UE::Landscape::EditLayers::IEditLayerRendererProvider
@@ -258,6 +258,14 @@ public:
 
 #endif // WITH_EDITOR
 
+	// BP accessors 
+
+	/**
+	* @return the name of the edit layer (editor-only : returns an invalid name at runtime)
+	*/
+	UFUNCTION(BlueprintCallable, DisplayName = GetName, Category = "Landscape|Editor", meta = (DevelopmentOnly))
+	FName GetNameBP() const;
+
 	UE_DEPRECATED(5.7, "Register custom context menu actions using RegisterCustomEditLayerClassLayout from ILandscapeEditorModule")
 PRAGMA_DISABLE_DEPRECATION_WARNINGS
 	virtual TArray<FEditLayerAction> GetActions() const 
diff --git a/Engine/Source/Runtime/Landscape/Private/LandscapeEdit.cpp b/Engine/Source/Runtime/Landscape/Private/LandscapeEdit.cpp
--- a/Engine/Source/Runtime/Landscape/Private/LandscapeEdit.cpp
+++ b/Engine/Source/Runtime/Landscape/Private/LandscapeEdit.cpp
@@ -8344,6 +8344,24 @@ TArray<FName> ALandscape::GetTargetLayerNames(bool bInIncludeVisibilityLayer) co
 	return Result;
 }
 
+const TArray<ULandscapeEditLayerBase*> ALandscape::GetEditLayersBP() const
+{
+#if WITH_EDITOR
+		return GetEditLayers();
+#else // WITH_EDITOR
+		return {};
+#endif // !WITH_EDITOR
+}
+
+const ULandscapeEditLayerBase* ALandscape::GetEditLayerByNameBP(const FName& InLayerName) const
+{
+#if WITH_EDITOR
+	return ObjectPtrWrap(GetEditLayer(InLayerName));
+#else // WITH_EDITOR
+	return nullptr;
+#endif // !WITH_EDITOR
+}
+
 
 #if WITH_EDITOR
 
diff --git a/Engine/Source/Runtime/Landscape/Private/LandscapeEditLayer.cpp b/Engine/Source/Runtime/Landscape/Private/LandscapeEditLayer.cpp
--- a/Engine/Source/Runtime/Landscape/Private/LandscapeEditLayer.cpp
+++ b/Engine/Source/Runtime/Landscape/Private/LandscapeEditLayer.cpp
@@ -411,6 +411,16 @@ void ULandscapeEditLayerBase::SetWeightmapLayerAllocationBlendInternal(const TMa
 
 #endif // WITH_EDITORONLY_DATA
 
+FName ULandscapeEditLayerBase::GetNameBP() const
+{
+#if WITH_EDITOR
+	return GetName();
+#else // WITH_EDITOR
+	return NAME_None;
+#endif // !WITH_EDITOR
+}
+
+
 // ----------------------------------------------------------------------------------
 
 bool ULandscapeEditLayer::SupportsTargetType(ELandscapeToolTargetType InType) const
