Landscape Patch : fixed AssignToLandscape BP function so that it can be used to unassign a landscape too and so that it resets the edit layer when the provided name is invalid

#Synced-CL 49593386

diff --git a/Engine/Plugins/Editor/LandscapePatch/Source/LandscapePatch/Private/LandscapePatchComponent.cpp b/Engine/Plugins/Editor/LandscapePatch/Source/LandscapePatch/Private/LandscapePatchComponent.cpp
--- a/Engine/Plugins/Editor/LandscapePatch/Source/LandscapePatch/Private/LandscapePatchComponent.cpp
+++ b/Engine/Plugins/Editor/LandscapePatch/Source/LandscapePatch/Private/LandscapePatchComponent.cpp
@@ -853,34 +853,50 @@ bool ULandscapePatchComponent::AssignToLandscape(ALandscape* InLandscape, const
 	bSuccess = BindToLandscape(InLandscape);
 	if (bSuccess && (InLandscape != nullptr))
 	{
-		const ULandscapeEditLayerBase* LocalEditLayer = InLandscape->GetEditLayerConst(InEditLayerName);
-		if (LocalEditLayer != nullptr)
+		// Passing an empty name is the equivalent of un-assigning the edit layer, it's considered valid : 
+		if (InEditLayerName.IsNone())
 		{
-			const ULandscapePatchEditLayer* PatchEditLayer = Cast<const ULandscapePatchEditLayer>(LocalEditLayer);
-			if (PatchEditLayer != nullptr)
+			bSuccess = BindToEditLayer(FGuid());
+		}
+		else
+		{
+			// Validate that the edit layer name corresponds to a valid patch edit layer and spew errors if not : 
+			const ULandscapeEditLayerBase* LocalEditLayer = InLandscape->GetEditLayerConst(InEditLayerName);
+			if (LocalEditLayer != nullptr)
 			{
-				bSuccess = BindToEditLayer(LocalEditLayer->GetGuid());
+				const ULandscapePatchEditLayer* PatchEditLayer = Cast<const ULandscapePatchEditLayer>(LocalEditLayer);
+				if (PatchEditLayer != nullptr)
+				{
+					// All good, we can try to attach to that edit layer : 
+					bSuccess = BindToEditLayer(LocalEditLayer->GetGuid());
+				}
+				else
+				{
+					UE_LOG(LogLandscapePatch, Warning, TEXT("Edit layer with the given name (%s) on landscape %s is not a patch edit layer (%s). "
+						"(Package: %s"),
+						*InEditLayerName.ToString(),
+						*InLandscape->GetActorLabel(),
+						*LocalEditLayer->GetClass()->GetName(),
+						*GetNameSafe(GetPackage()));
+
+					// Un-assign the edit layer so that the function call effectively changes the edit layer according to the error :
+					BindToEditLayer(FGuid());
+					bSuccess = false;
+				}
 			}
 			else
 			{
-				UE_LOG(LogLandscapePatch, Warning, TEXT("Edit layer with the given name (%s) on landscape %s is not a patch edit layer (%s). "
+				UE_LOG(LogLandscapePatch, Warning, TEXT("Could not find an edit layer with the given name (%s) on landscape %s for the patch. "
 					"(Package: %s"),
 					*InEditLayerName.ToString(),
 					*InLandscape->GetActorLabel(),
-					*LocalEditLayer->GetClass()->GetName(),
 					*GetNameSafe(GetPackage()));
+
+				// Un-assign the edit layer so that the function call effectively changes the edit layer according to the error :
+				BindToEditLayer(FGuid());
 				bSuccess = false;
 			}
 		}
-		else
-		{
-			UE_LOG(LogLandscapePatch, Warning, TEXT("Could not find an edit layer with the given name (%s) on landscape %s for the patch. "
-				"(Package: %s"),
-				*InEditLayerName.ToString(),
-				*InLandscape->GetActorLabel(),
-				*GetNameSafe(GetPackage()));
-			bSuccess = false;
-		}
 	}
 #endif
 
