Landscape Patch - Exposed Texture Transform properties to blueprint
* Added GetHeightAlphaRenderTarget/GetWeightPatchAlphaRender target Blueprint functions
* Improved bUseAlpha deprecation messages

#Synced-CL 49736387

diff --git a/Engine/Plugins/Editor/LandscapePatch/Source/LandscapePatch/Private/LandscapeTexturePatch.cpp b/Engine/Plugins/Editor/LandscapePatch/Source/LandscapePatch/Private/LandscapeTexturePatch.cpp
--- a/Engine/Plugins/Editor/LandscapePatch/Source/LandscapePatch/Private/LandscapeTexturePatch.cpp
+++ b/Engine/Plugins/Editor/LandscapePatch/Source/LandscapePatch/Private/LandscapeTexturePatch.cpp
@@ -211,7 +211,7 @@ FMatrix44d FLandscapeTexturePatchTextureTransform::GetTransformToUV(const FTrans
 	return FTransform(FRotator(0.0, Rotation, 0.0), UVOffset, bUseWorldSpaceScale ? FVector(UVWorldScale.X, UVWorldScale.Y, 1.0) : UVRelativeScale).ToMatrixWithScale();
 }
 
-UTextureRenderTarget2D* FLandscapeTexturePatchAlphaSettings::GetAlphaRenderTarget(UObject* InOuter)
+UTextureRenderTarget2D* FLandscapeTexturePatchAlphaSettings::GetAlphaRenderTarget(UObject* InOuter, bool bAlwaysMarkDirty/*=true*/)
 {
 #if WITH_EDITOR
 	// InOuter is the texture patch or weight patch layer info
@@ -225,7 +225,7 @@ UTextureRenderTarget2D* FLandscapeTexturePatchAlphaSettings::GetAlphaRenderTarge
 	{
 		if (!InternalAlphaData || !InternalAlphaData->GetRenderTarget())
 		{
-			SetAlphaSourceMode(AlphaSourceMode, InOuter);
+			SetAlphaSourceMode(AlphaSourceMode, InOuter, bAlwaysMarkDirty);
 		}
 
 		return ensure(InternalAlphaData) ? InternalAlphaData->GetRenderTarget() : nullptr;
@@ -235,7 +235,7 @@ UTextureRenderTarget2D* FLandscapeTexturePatchAlphaSettings::GetAlphaRenderTarge
 	return nullptr;
 }
 
-UTexture2D* FLandscapeTexturePatchAlphaSettings::GetAlphaInternalTexture(UObject* InOuter)
+UTexture2D* FLandscapeTexturePatchAlphaSettings::GetAlphaInternalTexture(UObject* InOuter, bool bAlwaysMarkDirty/*=true*/)
 {
 #if WITH_EDITOR
 	// InOuter is the texture patch or weight patch layer info
@@ -250,7 +250,7 @@ UTexture2D* FLandscapeTexturePatchAlphaSettings::GetAlphaInternalTexture(UObject
 	{
 		if (!InternalAlphaData || !InternalAlphaData->GetInternalTexture())
 		{
-			SetAlphaSourceMode(AlphaSourceMode, InOuter);
+			SetAlphaSourceMode(AlphaSourceMode, InOuter, bAlwaysMarkDirty);
 		}
 
 		return ensure(InternalAlphaData) ? InternalAlphaData->GetInternalTexture() : nullptr;
@@ -297,7 +297,8 @@ void FLandscapeTexturePatchAlphaSettings::SetAlphaSourceMode(ELandscapeTexturePa
 			ULandscapeWeightTextureBackedRenderTarget* InternalDataToReturn = NewObject<ULandscapeWeightTextureBackedRenderTarget>(InOuter);
 			InternalDataToReturn->SetFlags(RF_Transactional);
 			InternalDataToReturn->SetResolution(Resolution.X, Resolution.Y);
-			InternalDataToReturn->SetUseAlphaChannel(true);
+			// Alpha texture only uses a single color channel and will use RTF_R8/TSF_G8 formats for internal texture/render targets
+			InternalDataToReturn->SetUseAlphaChannel(false);
 			return InternalDataToReturn;
 		});
 
@@ -1214,6 +1215,9 @@ void ULandscapeTexturePatch::ReinitializeWeightPatch(ULandscapeWeightPatchTextur
 		}
 	}
 
+	// Clear the default/shared weight patch alpha
+	WeightPatchesAlphaSettings.SetAlphaSourceMode(ELandscapeTexturePatchAlphaSourceMode::None, this, /*bAlwaysMarkDiry =*/true);
+
 	// Clear the alpha overrides
 	PatchInfo->SetWeightPatchAlphaSourceMode(ELandscapeTexturePatchAlphaSourceMode::None, /*bInOverrideAlpha =*/false);
 	
@@ -1966,10 +1970,10 @@ void ULandscapeWeightPatchTextureInfo::SetWeightPatchAlphaSourceMode(ELandscapeT
 		WeightPatchAlphaSettings.SetAlphaSourceMode(ELandscapeTexturePatchAlphaSourceMode::None, this);
 	}
 
-	// The primary internal source texture tracks the alpha status
+	// Weightmap alpha channel is only utilized when SourceTextureChannel (weightmap's alpha) is selected. All other alpha modes generate a second alpha specific texture
 	if (InternalData)
 	{
-		InternalData->SetUseAlphaChannel(InSourceMode != ELandscapeTexturePatchAlphaSourceMode::None);
+		InternalData->SetUseAlphaChannel(InSourceMode == ELandscapeTexturePatchAlphaSourceMode::SourceTextureChannel);
 	}
 #endif // WITH_EDITOR
 }
@@ -2165,6 +2169,25 @@ void ULandscapeTexturePatch::TransitionHeightSourceModeInternal(ELandscapeTextur
 }
 #endif
 
+UTextureRenderTarget2D* ULandscapeTexturePatch::GetHeightAlphaRenderTarget(bool bMarkDirty)
+{
+#if WITH_EDITOR
+	if (IsTemplate())
+	{
+		return nullptr;
+	}
+
+	if (bMarkDirty)
+	{
+		MarkPackageDirty();
+	}
+
+	return HeightAlphaSettings.GetAlphaRenderTarget(this, bMarkDirty);
+#else
+	return nullptr;
+#endif
+}
+
 void ULandscapeTexturePatch::SetHeightTextureAsset(UTexture* TextureIn)
 {
 	if (TextureIn && TextureIn->VirtualTextureStreaming != 0)
@@ -2304,8 +2327,15 @@ void ULandscapeTexturePatch::ResetHeightEncodingMode(ELandscapeTextureHeightPatc
 #if WITH_EDITOR
 void ULandscapeTexturePatch::ResetHeightRenderTargetFormat()
 {
-	SetHeightRenderTargetFormat(HeightEncoding == ELandscapeTextureHeightPatchEncoding::NativePackedHeight ? ETextureRenderTargetFormat::RTF_RGBA8
-		: (HeightAlphaSettings.GetAlphaSourceMode() != ELandscapeTexturePatchAlphaSourceMode::None) ? ETextureRenderTargetFormat::RTF_RGBA32f : ETextureRenderTargetFormat::RTF_R32f);
+	// Heightmap alpha channel is only utilized when SourceTextureChannel (heightmap's alpha) is selected. All other alpha modes generate a second alpha specific texture
+	const bool bUsesAlphaChannel = HeightAlphaSettings.GetAlphaSourceMode() == ELandscapeTexturePatchAlphaSourceMode::SourceTextureChannel;
+	const ETextureRenderTargetFormat DefaultEncodingFormat = bUsesAlphaChannel ? ETextureRenderTargetFormat::RTF_RGBA32f : ETextureRenderTargetFormat::RTF_R32f;
+
+	const ETextureRenderTargetFormat NewFormat = (HeightEncoding == ELandscapeTextureHeightPatchEncoding::NativePackedHeight) 
+		? ETextureRenderTargetFormat::RTF_RGBA8
+		: DefaultEncodingFormat;
+
+	SetHeightRenderTargetFormat(NewFormat);
 
 	// When the height uses internal texture / native packed height we limit the user to RG texture color channels
 	if (HeightEncoding == ELandscapeTextureHeightPatchEncoding::NativePackedHeight || HeightSourceMode == ELandscapeTexturePatchSourceMode::InternalTexture)
@@ -2397,7 +2427,7 @@ void ULandscapeTexturePatch::CreateWeightPatch(const FName& InWeightmapLayerName
 		NewWeightPatch->InternalData = NewObject<ULandscapeWeightTextureBackedRenderTarget>(NewWeightPatch);
 		NewWeightPatch->InternalData->SetFlags(RF_Transactional);
 		NewWeightPatch->InternalData->SetResolution(ResolutionX, ResolutionY);
-		NewWeightPatch->InternalData->SetUseAlphaChannel(InAlphaSourceMode != ELandscapeTexturePatchAlphaSourceMode::None);
+		NewWeightPatch->InternalData->SetUseAlphaChannel(InAlphaSourceMode == ELandscapeTexturePatchAlphaSourceMode::SourceTextureChannel);
 		NewWeightPatch->InternalData->Initialize();
 	}
 	WeightPatches.Add(NewWeightPatch);
@@ -2512,6 +2542,29 @@ const UTexture* ULandscapeTexturePatch::GetWeightPatchAlphaTextureAsset(const FN
 	return WeightPatchesAlphaSettings.AlphaTextureAsset;
 }
 
+UTextureRenderTarget2D* ULandscapeTexturePatch::GetWeightPatchAlphaRenderTarget(const FName& InWeightmapLayerName, bool bMarkDirty)
+{
+	if (IsTemplate())
+	{
+		return nullptr;
+	}
+
+	for (const TObjectPtr<ULandscapeWeightPatchTextureInfo>& WeightPatch : WeightPatches)
+	{
+		if (IsValid(WeightPatch) && WeightPatch->WeightmapLayerName == InWeightmapLayerName)
+		{
+			if (bMarkDirty)
+			{
+				MarkPackageDirty();
+			}
+
+			// returns the default (shared) weight patch render target or the specific override weight patch render target
+			return GetWeightPatchAlphaRenderTarget(WeightPatch, bMarkDirty);
+		}
+	}
+	return nullptr;
+}
+
 void ULandscapeTexturePatch::SetWeightPatchAlphaTextureAsset(const FName& InWeightmapLayerName, UTexture* InTexture, bool bInOverrideAlpha, bool bAlwaysMarkDirty)
 {
 	if (InTexture && InTexture->VirtualTextureStreaming != 0)
@@ -2816,14 +2869,38 @@ const FLandscapeTexturePatchAlphaSettings& ULandscapeTexturePatch::GetWeightPatc
 	return WeightPatch->bOverrideWeightAlphaSettings ? WeightPatch->WeightPatchAlphaSettings : WeightPatchesAlphaSettings;
 }
 
-UTextureRenderTarget2D* ULandscapeTexturePatch::GetWeightPatchAlphaRenderTarget(ULandscapeWeightPatchTextureInfo* WeightPatch)
+UTextureRenderTarget2D* ULandscapeTexturePatch::GetWeightPatchAlphaRenderTarget(ULandscapeWeightPatchTextureInfo* WeightPatch, bool bMarkDirty)
 {
-	return GetWeightPatchAlphaSettings(WeightPatch).GetAlphaRenderTarget(WeightPatch);
+	if (IsTemplate() || !IsValid(WeightPatch))
+	{
+		return nullptr;
+	}
+
+	if (WeightPatch->bOverrideWeightAlphaSettings)
+	{
+		return WeightPatch->WeightPatchAlphaSettings.GetAlphaRenderTarget(WeightPatch, bMarkDirty);
+	}
+	else
+	{
+		return WeightPatchesAlphaSettings.GetAlphaRenderTarget(this, bMarkDirty);
+	}
 }
 
-UTexture2D* ULandscapeTexturePatch::GetWeightPatchAlphaInternalTexture(ULandscapeWeightPatchTextureInfo* WeightPatch)
+UTexture2D* ULandscapeTexturePatch::GetWeightPatchAlphaInternalTexture(ULandscapeWeightPatchTextureInfo* WeightPatch, bool bMarkDirty)
 {
-	return GetWeightPatchAlphaSettings(WeightPatch).GetAlphaInternalTexture(WeightPatch);
+	if (IsTemplate() || !IsValid(WeightPatch))
+	{
+		return nullptr;
+	}
+
+	if (WeightPatch->bOverrideWeightAlphaSettings)
+	{
+		return WeightPatch->WeightPatchAlphaSettings.GetAlphaInternalTexture(WeightPatch, bMarkDirty);
+	}
+	else
+	{
+		return WeightPatchesAlphaSettings.GetAlphaInternalTexture(this, bMarkDirty);
+	}
 }
 
 UTexture2D* ULandscapeTexturePatch::GetWeightPatchInternalTexture(ULandscapeWeightPatchTextureInfo* WeightPatch)
diff --git a/Engine/Plugins/Editor/LandscapePatch/Source/LandscapePatch/Public/LandscapeTexturePatch.h b/Engine/Plugins/Editor/LandscapePatch/Source/LandscapePatch/Public/LandscapeTexturePatch.h
--- a/Engine/Plugins/Editor/LandscapePatch/Source/LandscapePatch/Public/LandscapeTexturePatch.h
+++ b/Engine/Plugins/Editor/LandscapePatch/Source/LandscapePatch/Public/LandscapeTexturePatch.h
@@ -188,21 +188,21 @@ struct FLandscapeTexturePatchTextureTransform
 	GENERATED_BODY()
 public:
 	/** If enabled, texture UVs are scaled using world-space units (1000 = tile texture every 10 meters). If disabled, scaling is applied in local UV space. */
-	UPROPERTY(EditAnywhere, Category = WeightPatch)
+	UPROPERTY(EditAnywhere, BlueprintReadWrite, Category = "Texture Transform")
 	bool bUseWorldSpaceScale = false;
 
 	/** Scale applied to the texture's UV. By default, scaling is applied in local UV space. Relative scale 2.0 = texture is twice as large as original size. 
 	* When World Space Scale is enabled, scale 1000 = tile texture every 10 meters. 
 	*/
-	UPROPERTY(EditAnywhere, Category = WeightPatch)
+	UPROPERTY(EditAnywhere, BlueprintReadWrite, Category = "Texture Transform")
 	FVector2D Scale = FVector2D(1, 1);
 	
 	/** Rotation in degrees applied to the texture's UV. */
-	UPROPERTY(EditAnywhere, Category = WeightPatch, meta = (UIMin = "0", UIMax = "360", ClampMin = "0", ClampMax = "360"))
+	UPROPERTY(EditAnywhere, BlueprintReadWrite, Category = "Texture Transform", meta = (UIMin = "0", UIMax = "360", ClampMin = "0", ClampMax = "360"))
 	float Rotation = 0;
 
 	/** Relative offset applied to the texture's UV. */
-	UPROPERTY(EditAnywhere, Category = WeightPatch, meta = (UIMin = "0", UIMax = "1", ClampMin = "0", ClampMax = "1"))
+	UPROPERTY(EditAnywhere, BlueprintReadWrite, Category = "Texture Transform", meta = (UIMin = "0", UIMax = "1", ClampMin = "0", ClampMax = "1"))
 	FVector2D Offset = FVector2D(0, 0);
 
 	FMatrix44d GetTransformToUV(const FTransform& PatchToWorldIn, const FVector2f& PatchWorldDimensionsIn) const;
@@ -222,9 +222,9 @@ struct FLandscapeTexturePatchAlphaSettings
 	GENERATED_BODY()
 public:
 	/** When alpha source mode is TextureBackedRenderTarget returns the alpha render target. Will allocate new render target if not yet initialized */
-	UTextureRenderTarget2D* GetAlphaRenderTarget(UObject* InOuter);
+	UTextureRenderTarget2D* GetAlphaRenderTarget(UObject* InOuter, bool bAlwaysMarkDirty = true);
 	/** When alpha source mode is  InternalTexture/TextureBackedRenderTarget returns the alpha internal texture. Will allocate new internal data if not yet initialized */
-	UTexture2D* GetAlphaInternalTexture(UObject* InOuter);
+	UTexture2D* GetAlphaInternalTexture(UObject* InOuter, bool bAlwaysMarkDirty = true);
 
 	/** Transitions from the previous alpha source mode to the new mode. Marks the Outer owner dirty and clears internal data if needed  
 	* @param InNewSourceMode the new alpha source mode
@@ -342,8 +342,8 @@ protected:
 	UPROPERTY(VisibleAnywhere, Category = WeightPatch, Instanced, AdvancedDisplay, meta = (DisplayName = "Debug Data"))
 	TObjectPtr<ULandscapeWeightTextureBackedRenderTarget> InternalData = nullptr;
 
-	UE_DEPRECATED(5.8, "bUseAlphaChannel is deprecated, use the WeightPatchAlphaSettings to define the alpha behavior.")
-	UPROPERTY(meta = (DeprecatedProperty, DeprecationMessage = "Use WeightPatchAlphaSettings instead. See ELandscapeTexturePatchAlphaSourceMode::SourceTextureChannel to enable legacy behavior"))
+	UE_DEPRECATED(5.8, "bUseAlphaChannel is deprecated, use the WeightPatchAlphaSettings.SetAlphaSourceMode() with ELandscapeTexturePatchAlphaSourceMode (use alpha mode None or SourceTextureChannel for legacy behavior).")
+	UPROPERTY(meta = (DeprecatedProperty, DeprecationMessage = "Use WeightPatchAlphaSettings.SetAlphaSourceMode() with ELandscapeTexturePatchAlphaSourceMode (use alpha mode None or SourceTextureChannel for legacy behavior)."))
 	bool bUseAlphaChannel_DEPRECATED = false;
 
 	UPROPERTY(EditAnywhere, Category = WeightPatch, meta = (InlineEditConditionToggle))
@@ -585,6 +585,22 @@ public:
 	UFUNCTION(BlueprintCallable, Category = LandscapePatch)
 	LANDSCAPEPATCH_API virtual void SetHeightAlphaTextureChannel(ELandscapeTexturePatchTextureChannel InTextureChannel, bool bAlwaysMarkDirty = true);
 
+	/**
+	* Gets the internal height alpha render target, if the height alpha source mode is set to Texture Backed Render Target
+	*
+	* Things that should be set up if using the internal render target:
+	* - An appropriate texture size should have been set with SetResolution. If the patch extent has already
+	*  been set, you can base your resolution on the extent and the resolution of the landscape by using
+	*  GetInitResolutionFromLandscape()
+	* - The alpha uses RenderTargetFormat RTF_R8 by default. The AlphaTextureChannel property defines the single color channel 
+	*   storing the Alpha data
+	*
+	* @param bMarkDirty If true, marks the containing package as dirty, since the render target is presumably
+	*  being written to. Can be set to false if the render target is not being written to.
+	*/
+	UFUNCTION(BlueprintCallable, Category = LandscapePatch)
+	LANDSCAPEPATCH_API virtual UTextureRenderTarget2D* GetHeightAlphaRenderTarget(bool bMarkDirty = true);
+
 	/**
 	 * Sets the texture used for height when the height source mode is set to texture asset. Note that
 	 * virtual textures are not supported.
@@ -617,8 +633,8 @@ public:
 	UFUNCTION(BlueprintCallable, Category = LandscapePatch, meta = (ETextureRenderTargetFormat = "ETextureRenderTargetFormat::RTF_R32f"))
 	LANDSCAPEPATCH_API void SetHeightRenderTargetFormat(ETextureRenderTargetFormat Format);
 
-	UE_DEPRECATED(all, "SetUseAlphaChannelForHeight is deprecated. Use SetHeightAlphaSourceMode with ELandscapeTexturePatchAlphaSourceMode instead")
-	UFUNCTION(meta = (DeprecatedFunction, DeprecationMessage = "SetUseAlphaChannelForHeight has been deprecated, use SetHeightAlphaSourceMode instead."))
+	UE_DEPRECATED(all, "SetUseAlphaChannelForHeight is deprecated. Use SetHeightAlphaSourceMode with ELandscapeTexturePatchAlphaSourceMode instead (use alpha mode None or SourceTextureChannel for legacy behavior)")
+	UFUNCTION(meta = (DeprecatedFunction, DeprecationMessage = "SetUseAlphaChannelForHeight has been deprecated, use SetHeightAlphaSourceMode instead (use alpha mode None or SourceTextureChannel for legacy behavior)."))
 	void SetUseAlphaChannelForHeight(bool bUse)
 	{
 		// Deprecated
@@ -675,8 +691,8 @@ public:
 
 
 	// Weight related functions:
-	UE_DEPRECATED(all, "Use CreateWeightPatch with ELandscapeTexturePatchAlphaSourceMode instead")
-	UFUNCTION(meta = (DeprecatedFunction, DeprecationMessage = "AddWeightPatch has been deprecated, use CreateWeightPatch with ELandscapeTexturePatchAlphaSourceMode instead."))
+	UE_DEPRECATED(all, "Use CreateWeightPatch with ELandscapeTexturePatchAlphaSourceMode instead (use alpha mode None or SourceTextureChannel for legacy behavior)")
+	UFUNCTION(meta = (DeprecatedFunction, DeprecationMessage = "AddWeightPatch has been deprecated, use CreateWeightPatch with ELandscapeTexturePatchAlphaSourceMode instead (use None or SourceTextureChannel for legacy behavior)."))
 	LANDSCAPEPATCH_API virtual void AddWeightPatch(const FName& InWeightmapLayerName, ELandscapeTexturePatchSourceMode SourceMode, bool bUseAlphaChannel); 
 
 	/**
@@ -741,8 +757,8 @@ public:
 	UFUNCTION(BlueprintCallable, Category = LandscapePatch)
 	LANDSCAPEPATCH_API virtual UTextureRenderTarget2D* GetWeightPatchRenderTarget(const FName& InWeightmapLayerName, bool bMarkDirty = true);
 
-	UE_DEPRECATED(all, "SetUseAlphaChannelForWeightPatch is deprecated. Use SetWeightPatchAlphaSourceMode with ELandscapeTexturePatchAlphaSourceMode instead")
-	UFUNCTION(meta = (DeprecatedFunction, DeprecationMessage = "SetUseAlphaChannelForWeightPatch has been deprecated, use SetWeightPatchAlphaSourceMode instead."))
+	UE_DEPRECATED(all, "SetUseAlphaChannelForWeightPatch is deprecated. Use SetWeightPatchAlphaSourceMode with ELandscapeTexturePatchAlphaSourceMode instead (use None or SourceTextureChannel for legacy behavior)")
+	UFUNCTION(meta = (DeprecatedFunction, DeprecationMessage = "SetUseAlphaChannelForWeightPatch has been deprecated, use SetWeightPatchAlphaSourceMode instead (use None or SourceTextureChannel for legacy behavior)."))
 	LANDSCAPEPATCH_API virtual void SetUseAlphaChannelForWeightPatch(const FName& InWeightmapLayerName, bool bUseAlphaChannel); 
 
 	UFUNCTION(BlueprintCallable, Category = LandscapePatch)
@@ -755,11 +771,25 @@ public:
 	 * @param bInOverrideAlpha sets the weightmap layer's override flag. When false, the weightmap layer will fallback to the default source mode
  	 * @param bAlwaysMarkDirty mark the package dirty when property changes
 	 */
+	UFUNCTION(BlueprintCallable, Category = LandscapePatch)
 	LANDSCAPEPATCH_API virtual void SetWeightPatchAlphaSourceMode(const FName& InWeightmapLayerName, ELandscapeTexturePatchAlphaSourceMode InSourceMode, bool bInOverrideAlpha = true, bool bAlwaysMarkDirty = true);
 
 	UFUNCTION(BlueprintCallable, Category = LandscapePatch)
 	LANDSCAPEPATCH_API virtual const UTexture* GetWeightPatchAlphaTextureAsset(const FName& InWeightmapLayerName) const;
 
+	/**
+ 	 * Gets the internal weight patch alpha render target, if the weight patch alpha source mode is set to Texture Backed Render Target
+	 * Returns a reference to the default (shared) weightmap alpha render target unless the given weightmap patch overrides the default alpha settings
+	 *
+	 * - The alpha uses RenderTargetFormat RTF_R8 by default. The AlphaTextureChannel property defines the single color channel 
+	 *   storing the Alpha data
+	 * 
+	 * @param bMarkDirty If true, marks the containing package as dirty, since the render target is presumably
+	 *  being written to. Can be set to false if the render target is not being written to.
+	 */
+	UFUNCTION(BlueprintCallable, Category = LandscapePatch)
+	LANDSCAPEPATCH_API virtual UTextureRenderTarget2D* GetWeightPatchAlphaRenderTarget(const FName& InWeightmapLayerName, bool bMarkDirty = true);
+
 	/**
 	 * Sets the default weight patch alpha texture asset for all weight patches or overrides the texture for a specific weight patch
 	 * @param InWeightmapLayerName weightmap layer name to set the alpha texture. Overrides a specific layer or sets the default texture when name is None
@@ -782,7 +812,6 @@ public:
 	UFUNCTION(BlueprintCallable, Category = LandscapePatch)
 	LANDSCAPEPATCH_API virtual void SetWeightPatchAlphaTextureTransform(const FName& InWeightmapLayerName, const FLandscapeTexturePatchTextureTransform& InTextureTransform, bool bAlwaysMarkDirty = true);
 
-
 	UFUNCTION(BlueprintCallable, Category = LandscapePatch)
 	LANDSCAPEPATCH_API virtual ELandscapeTexturePatchTextureChannel GetWeightPatchAlphaTextureChannel(const FName& InWeightmapLayerName) const;
 
@@ -872,8 +901,8 @@ protected:
 	UPROPERTY(EditAnywhere, Category = HeightPatch, meta = (DisplayName = "Texture Transform"), meta = (EditConditionHides, EditCondition = "HeightSourceMode != ELandscapeTexturePatchSourceMode::None"))
 	FLandscapeTexturePatchTextureTransform HeightTextureUVTransform;
 
-	UE_DEPRECATED(5.8, "bUseAlphaChannel is deprecated, use the ELandscapeTexturePatchAlphaSourceMode to define the alpha behavior.")
-	UPROPERTY(meta = (DeprecatedProperty, DeprecationMessage = "Use HeightAlphaSettings instead. See ELandscapeTexturePatchAlphaSourceMode::SourceTextureChannel to enable legacy behavior"))
+	UE_DEPRECATED(5.8, "bUseAlphaChannel is deprecated, HeightAlphaSettings.SetAlphaSourceMode() with ELandscapeTexturePatchAlphaSourceMode (use None or SourceTextureChannel for legacy behavior).")
+	UPROPERTY(meta = (DeprecatedProperty, DeprecationMessage = "Use HeightAlphaSettings.SetAlphaSourceMode() with ELandscapeTexturePatchAlphaSourceMode instead (use None or SourceTextureChannel for legacy behavior)"))
 	bool bUseTextureAlphaForHeight_DEPRECATED = false;
 
 	UPROPERTY(EditAnywhere, Category = HeightPatch, meta = (EditConditionHides, EditCondition = "HeightSourceMode != ELandscapeTexturePatchSourceMode::None"))
@@ -979,9 +1008,9 @@ protected:
 	LANDSCAPEPATCH_API UTextureRenderTarget2D* GetWeightPatchRenderTarget(ULandscapeWeightPatchTextureInfo* WeightPatch);
 	LANDSCAPEPATCH_API UTexture2D* GetWeightPatchInternalTexture(ULandscapeWeightPatchTextureInfo* WeightPatch);
 	/** Returns the default weight patch alpha render target or the specific weight patch override render target */
-	LANDSCAPEPATCH_API UTextureRenderTarget2D* GetWeightPatchAlphaRenderTarget(ULandscapeWeightPatchTextureInfo* WeightPatch);
+	LANDSCAPEPATCH_API UTextureRenderTarget2D* GetWeightPatchAlphaRenderTarget(ULandscapeWeightPatchTextureInfo* WeightPatch, bool bMarkDirty = true);
 	/** Returns the default weight patch alpha internal texture or the specific weight patch override internal texture */
-	LANDSCAPEPATCH_API UTexture2D* GetWeightPatchAlphaInternalTexture(ULandscapeWeightPatchTextureInfo* WeightPatch);
+	LANDSCAPEPATCH_API UTexture2D* GetWeightPatchAlphaInternalTexture(ULandscapeWeightPatchTextureInfo* WeightPatch, bool bMarkDirty = true);
 
 private:
 	void UpdateHeightConvertToNativeParamsIfNeeded();
