Fixup API errors

#rnx

#Synced-CL 50159792

#Synced-CL 50159792

diff --git a/Engine/Plugins/Animation/IKRig/Source/IKRig/Public/Retargeter/RetargetOps/SpeedPlantingOp.h b/Engine/Plugins/Animation/IKRig/Source/IKRig/Public/Retargeter/RetargetOps/SpeedPlantingOp.h
--- a/Engine/Plugins/Animation/IKRig/Source/IKRig/Public/Retargeter/RetargetOps/SpeedPlantingOp.h
+++ b/Engine/Plugins/Animation/IKRig/Source/IKRig/Public/Retargeter/RetargetOps/SpeedPlantingOp.h
@@ -152,7 +152,7 @@ public:
 		const FTransform& InComponentTransform,
 		const double InComponentScale,
 		const FIKRetargetDebugDrawState& InEditorState) const override;
-	UE_API virtual bool HasDebugDrawing() override { return true; };
+	virtual bool HasDebugDrawing() override { return true; };
 	UE_API virtual FText GetWarningMessage() const override;
 	UE_API virtual void ResetChainSettingsToDefault(const FName& InChainName) override;
 	UE_API virtual bool AreChainSettingsAtDefault(const FName& InChainName) override;
diff --git a/Engine/Plugins/Editor/LandscapePatch/Source/LandscapePatch/Public/LandscapeTexturePatch.h b/Engine/Plugins/Editor/LandscapePatch/Source/LandscapePatch/Public/LandscapeTexturePatch.h
--- a/Engine/Plugins/Editor/LandscapePatch/Source/LandscapePatch/Public/LandscapeTexturePatch.h
+++ b/Engine/Plugins/Editor/LandscapePatch/Source/LandscapePatch/Public/LandscapeTexturePatch.h
@@ -531,7 +531,7 @@ public:
 	LANDSCAPEPATCH_API virtual void SetHeightSourceMode(ELandscapeTexturePatchSourceMode NewMode);
 
 	UFUNCTION(BlueprintCallable, Category = LandscapePatch)
-	LANDSCAPEPATCH_API virtual const FLandscapeTexturePatchTextureTransform& GetHeightTextureTransform() const { return HeightTextureUVTransform; };
+	virtual const FLandscapeTexturePatchTextureTransform& GetHeightTextureTransform() const { return HeightTextureUVTransform; };
 
 	/**
 	 * Sets the height texture's UV transform. Transform is applied relative UVSpace unless struct property bInWorldSpaceScale is true
@@ -564,7 +564,7 @@ public:
 	LANDSCAPEPATCH_API virtual void SetHeightAlphaSourceMode(ELandscapeTexturePatchAlphaSourceMode InNewMode, bool bAlwaysMarkDirty = true);
 
 	UFUNCTION(BlueprintCallable, Category = LandscapePatch)
-	LANDSCAPEPATCH_API virtual const FLandscapeTexturePatchTextureTransform& GetHeightAlphaTextureTransform() const { return HeightAlphaSettings.AlphaTextureUVTransform; };
+	virtual const FLandscapeTexturePatchTextureTransform& GetHeightAlphaTextureTransform() const { return HeightAlphaSettings.AlphaTextureUVTransform; };
 
 	/**
 	 * Sets the height alpha texture's UV transform. Transform is applied relative UVSpace unless struct property bInWorldSpaceScale is true
diff --git a/Engine/Plugins/Experimental/GameplayTargetingSystem/Source/GameplayTargetingSystem/Public/Tasks/TargetingSelectionTask_AOE.h b/Engine/Plugins/Experimental/GameplayTargetingSystem/Source/GameplayTargetingSystem/Public/Tasks/TargetingSelectionTask_AOE.h
--- a/Engine/Plugins/Experimental/GameplayTargetingSystem/Source/GameplayTargetingSystem/Public/Tasks/TargetingSelectionTask_AOE.h
+++ b/Engine/Plugins/Experimental/GameplayTargetingSystem/Source/GameplayTargetingSystem/Public/Tasks/TargetingSelectionTask_AOE.h
@@ -59,7 +59,7 @@ public:
 	/** Debug draws the outlines of the set shape type. */
 	UE_API void DebugDrawBoundingVolume(const FTargetingRequestHandle& TargetingHandle, const FColor& Color, const FOverlapDatum* OverlapDatum = nullptr) const;
 
-	UE_API ETargetingAOEShape GetShapeType() const { return ShapeType; }
+	ETargetingAOEShape GetShapeType() const { return ShapeType; }
 	UE_API void SetShapeType(ETargetingAOEShape InShapeType);
 	UE_API void SetHalfExtent(FVector InHalfExtent);
 	UE_API void SetRadius(FScalableFloat InRadius);
@@ -84,7 +84,7 @@ public:
 	/** Helper method to build the Collision Shape */
 	UE_API FCollisionShape GetCollisionShape() const;
 
-	UE_API bool GetUseRelativeOffset() const { return bUseRelativeOffset; }
+	bool GetUseRelativeOffset() const { return bUseRelativeOffset; }
 
 private:
 	/** Method to process the trace task immediately */
diff --git a/Engine/Plugins/FX/Niagara/Source/Niagara/Public/NiagaraSharedDataUtilities.h b/Engine/Plugins/FX/Niagara/Source/Niagara/Public/NiagaraSharedDataUtilities.h
--- a/Engine/Plugins/FX/Niagara/Source/Niagara/Public/NiagaraSharedDataUtilities.h
+++ b/Engine/Plugins/FX/Niagara/Source/Niagara/Public/NiagaraSharedDataUtilities.h
@@ -13,7 +13,7 @@ Base class for generic data items held by data managers.
 Niagara Systems, NDCs and potentially external users will hold onto these when they are being used.
 Allows Data Managers to ensure lifetime of their internal data that may be being reference inside Niagara and by external users.
 */
-struct NIAGARA_API FNiagaraSharedData : public TSharedFromThis<FNiagaraSharedData>
+struct FNiagaraSharedData : public TSharedFromThis<FNiagaraSharedData>
 {
 	virtual ~FNiagaraSharedData(){}
 
@@ -44,4 +44,4 @@ namespace NiagaraSharedDataUtilities
 									Inside Niagara, access using an invalid ID is safe but will fail and return an error in the DI function(s).
 	 */
 	NIAGARA_API FNiagaraSharedDataPtr GetOrCreateSharedData_StaticMesh(FNiagaraWorldManager* WorldMan, UStaticMeshComponent* StaticMeshComponent, TArrayView<UNiagaraComponent*> ReferencingComponents);
-};
\ No newline at end of file
+};
diff --git a/Engine/Plugins/PCG/Source/PCG/Public/Data/PCGWorldData.h b/Engine/Plugins/PCG/Source/PCG/Public/Data/PCGWorldData.h
--- a/Engine/Plugins/PCG/Source/PCG/Public/Data/PCGWorldData.h
+++ b/Engine/Plugins/PCG/Source/PCG/Public/Data/PCGWorldData.h
@@ -393,7 +393,7 @@ private:
 
 public:
 	UE_DEPRECATED(5.8, "Use Initialize taking a IPCGGraphExecutionSource parameter instead")
-	PCG_API void Initialize(UWorld* InWorld, const FBox& InBounds = FBox(EForceInit::ForceInit)) {}
+	void Initialize(UWorld* InWorld, const FBox& InBounds = FBox(EForceInit::ForceInit)) {}
 
 PRAGMA_DISABLE_DEPRECATION_WARNINGS
 	UE_DEPRECATED(5.8, "Use GetWorld() instead")
@@ -484,7 +484,7 @@ private:
 
 public:
 	UE_DEPRECATED(5.8, "Use Initialize taking a IPCGGraphExecutionSource parameter instead")
-	PCG_API void Initialize(UWorld* InWorld, const FTransform& InTransform, const FBox& InBounds = FBox(EForceInit::ForceInit), const FBox& InLocalBounds = FBox(EForceInit::ForceInit)) {}
+	void Initialize(UWorld* InWorld, const FTransform& InTransform, const FBox& InBounds = FBox(EForceInit::ForceInit), const FBox& InLocalBounds = FBox(EForceInit::ForceInit)) {}
 
 PRAGMA_DISABLE_DEPRECATION_WARNINGS
 	UE_DEPRECATED(5.8, "Use GetWorld() instead")
diff --git a/Engine/Plugins/PCG/Source/PCG/Public/PCGGraphExecutionInspection.h b/Engine/Plugins/PCG/Source/PCG/Public/PCGGraphExecutionInspection.h
--- a/Engine/Plugins/PCG/Source/PCG/Public/PCGGraphExecutionInspection.h
+++ b/Engine/Plugins/PCG/Source/PCG/Public/PCGGraphExecutionInspection.h
@@ -80,7 +80,7 @@ public:
 	PCG_API void AddReferencedObjects(FReferenceCollector& Collector);
 
 	UE_DEPRECATED(5.8, "Use GetInspectionDataPtr/InspectData instead")
-	PCG_API const FPCGDataCollection* GetInspectionData(const FPCGStack& InStack) const { return nullptr;}
+	const FPCGDataCollection* GetInspectionData(const FPCGStack& InStack) const { return nullptr;}
 private:
 	int32 InspectionCounter = 0;
 
diff --git a/Engine/Plugins/Runtime/NetworkPrediction/Source/NetworkPrediction/Public/NetworkPredictionReplicationProxy.h b/Engine/Plugins/Runtime/NetworkPrediction/Source/NetworkPrediction/Public/NetworkPredictionReplicationProxy.h
--- a/Engine/Plugins/Runtime/NetworkPrediction/Source/NetworkPrediction/Public/NetworkPredictionReplicationProxy.h
+++ b/Engine/Plugins/Runtime/NetworkPrediction/Source/NetworkPrediction/Public/NetworkPredictionReplicationProxy.h
@@ -50,7 +50,7 @@ struct FReplicationProxy
 	UE_API void OnPreReplication();	
 	UE_API bool Identical(const FReplicationProxy* Other, uint32 PortFlags) const;
 
-	UE_API int32 GetCachedPendingFrame() const 
+	int32 GetCachedPendingFrame() const 
 	{ 
 		return CachedPendingFrame; 
 	}
@@ -131,11 +131,11 @@ struct FServerReplicationRPCParameter
 
 	// Send flow: ctor(Proxy) -> SerializeToProxy
 	FServerReplicationRPCParameter(FReplicationProxy& InProxy) : Proxy(&InProxy) { }
-	UE_API void SetReplicationProxy(FReplicationProxy& InProxy) { Proxy = &InProxy; }
+	void SetReplicationProxy(FReplicationProxy& InProxy) { Proxy = &InProxy; }
 	UE_API bool SerializeToProxy(FArchive& Ar, class UPackageMap* Map, bool& bOutSuccess);
 	UE_API bool NetSerialize(FArchive& Ar, class UPackageMap* Map, bool& bOutSuccess);
 
-	UE_API UPackageMap* GetPackageMap() const
+	UPackageMap* GetPackageMap() const
 	{ 
 		return CachedPackageMap; 
 	}
