Landscape Patch - Added ability for a patch edit layer's heightmap/weightmap alpha values to impact texture/circle patches
* Negative heightmap alpha inverts patch

#Synced-CL 49736387

diff --git a/Engine/Plugins/Editor/LandscapePatch/Shaders/Private/LandscapeCircleHeightPatchPS.usf b/Engine/Plugins/Editor/LandscapePatch/Shaders/Private/LandscapeCircleHeightPatchPS.usf
--- a/Engine/Plugins/Editor/LandscapePatch/Shaders/Private/LandscapeCircleHeightPatchPS.usf
+++ b/Engine/Plugins/Editor/LandscapePatch/Shaders/Private/LandscapeCircleHeightPatchPS.usf
@@ -8,6 +8,7 @@
 float3 InCenter;
 float InRadius;
 float InFalloff;
+float InEditLayerWeightmapAlpha;
 
 #if CIRCLE_HEIGHT_PATCH
 
@@ -27,7 +28,7 @@ void ApplyLandscapeCircleHeightPatch(in float4 SVPos : SV_POSITION, out float4 O
 	{
 		Alpha = 1 - min((Distance - InRadius) / ClampedFalloff, 1.0);
 	}
-	
+
 	OutColor = float4(PackHeight(PatchHeight), PackHeightAlpha(Alpha, EHEIGHTMAPALPHAFLAGS_ALPHABLEND));
 }
 
@@ -39,8 +40,12 @@ void ApplyLandscapeCircleVisibilityPatch(in float4 SVPos : SV_POSITION, out floa
 {
 	int2 DestinationTextureCoordinates = floor(SVPos.xy);
 	float Distance = distance(DestinationTextureCoordinates, InCenter.xy);
-	
-	OutColor = (Distance <= InRadius) ? 1.0f : 0.0f;
+
+	// Apply the edit layer's weightmap alpha
+	// Visibility acts a binary toggle - affects the landscape if only if alpha is greater than 0.66f
+	float Alpha = 1.0f * InEditLayerWeightmapAlpha;
+
+	OutColor = (Distance <= InRadius) ? float4(1.0f, 1.0f, Alpha, Alpha) : 0.0f;
 }
 
 #endif // CIRCLE_VISIBILITY_PATCH
diff --git a/Engine/Plugins/Editor/LandscapePatch/Shaders/Private/LandscapeTexturePatchPS.usf b/Engine/Plugins/Editor/LandscapePatch/Shaders/Private/LandscapeTexturePatchPS.usf
--- a/Engine/Plugins/Editor/LandscapePatch/Shaders/Private/LandscapeTexturePatchPS.usf
+++ b/Engine/Plugins/Editor/LandscapePatch/Shaders/Private/LandscapeTexturePatchPS.usf
@@ -147,6 +147,9 @@ uint InAlphaMode;
 Texture2D<float4> InHeightPatchAlphaTexture;
 uint InAlphaTextureChannel;
 float4x4 InAlphaTextureTransformToUV;
+// Assigned edit layer's heightmap/weightmap alpha
+float InEditLayerHeightmapAlpha;
+float InEditLayerWeightmapAlpha;
 
 uint ConvertPatchBlendModeToHeightAlphaFlags()
 {
@@ -196,7 +199,8 @@ void ApplyLandscapeTextureHeightPatch(in float4 SVPos : SV_POSITION, out float4
 	// Apply the alpha from the source height texture or use the standalone alpha texture
 	Alpha *= GetAlphaFactor(InAlphaMode, InAlphaTextureChannel, PatchSampledValue, PatchAlphaSampledValue);
 
-	float OutputHeight = LANDSCAPE_MID_VALUE + PatchSignedHeight;
+	// Apply the edit layer's heightmap alpha (range [-1, 1]) directly to height to support negative range (negative alpha inverts the heightmap effect)
+	float OutputHeight = LANDSCAPE_MID_VALUE + (PatchSignedHeight * InEditLayerHeightmapAlpha);
 	OutColor = float4(PackHeight(OutputHeight), PackHeightAlpha(Alpha, ConvertPatchBlendModeToHeightAlphaFlags()));
 }
 #endif // APPLY_HEIGHT_PATCH
@@ -267,6 +271,8 @@ uint InAlphaMode;
 Texture2D<float4> InWeightPatchAlphaTexture;
 uint InAlphaTextureChannel;
 float4x4 InAlphaTextureTransformToUV;
+// Patch edit layer's weightmap alpha
+float InEditLayerWeightmapAlpha;
 
 uint ConvertPatchBlendModeToWeightAlphaFlags()
 {
@@ -312,6 +318,9 @@ void ApplyLandscapeTextureWeightPatch(in float4 SVPos : SV_POSITION, out float4
 	float4 PatchAlphaSampledValue = InWeightPatchAlphaTexture.Sample(InWeightPatchSampler, PatchAlphaUVCoordinatesTransformed);
 	Alpha *= GetAlphaFactor(InAlphaMode, InAlphaTextureChannel, PatchSampledValue, PatchAlphaSampledValue);
 	
+	// Apply the edit layer's weightmap alpha (range [0, 1])
+	Alpha *= InEditLayerWeightmapAlpha;
+
 	OutColor = float4(PackWeight(PatchWeight), PackWeightAlpha(Alpha, ConvertPatchBlendModeToWeightAlphaFlags()));
 }
 
diff --git a/Engine/Plugins/Editor/LandscapePatch/Source/LandscapePatch/Private/LandscapeCircleHeightPatch.cpp b/Engine/Plugins/Editor/LandscapePatch/Source/LandscapePatch/Private/LandscapeCircleHeightPatch.cpp
--- a/Engine/Plugins/Editor/LandscapePatch/Source/LandscapePatch/Private/LandscapeCircleHeightPatch.cpp
+++ b/Engine/Plugins/Editor/LandscapePatch/Source/LandscapePatch/Private/LandscapeCircleHeightPatch.cpp
@@ -137,8 +137,17 @@ bool ULandscapeCircleHeightPatch::ApplyCirclePatch(UE::Landscape::EditLayers::FR
 		return false;
 	}
 
+	const ULandscapePatchEditLayer* PatchEditLayer = GetBoundEditLayer();
+	check(PatchEditLayer != nullptr);
+	const float EditLayerHeightmapAlpha = PatchEditLayer->GetAlphaForTargetType(ELandscapeToolTargetType::Heightmap);
+	const float EditLayerWeightmapAlpha = PatchEditLayer->GetAlphaForTargetType(ELandscapeToolTargetType::Weightmap);
+
 	double ToHeightmapRadiusScale = GetComponentTransform().GetScale3D().X / HeightmapCoordsToWorld.GetScale3D().X;
 	FVector3d CircleCenterWorld = GetComponentTransform().GetTranslation();
+
+	// Apply patch's edit layer alpha (range [-1, 1]) to the height (negative alpha inverts heightmap effect)
+	CircleCenterWorld.Z *= EditLayerHeightmapAlpha;
+
 	FVector3d CenterInHeightmapCoordinates = HeightmapCoordsToWorld.InverseTransformPosition(CircleCenterWorld);
 	float RadiusAdjustment = bExclusiveRadius ? 0 : 1;
 	float HeightmapRadius = Radius * ToHeightmapRadiusScale + RadiusAdjustment;
@@ -186,6 +195,7 @@ bool ULandscapeCircleHeightPatch::ApplyCirclePatch(UE::Landscape::EditLayers::FR
 		, CenterInHeightmapCoordinates
 		, HeightmapRadius
 		, HeightmapFalloff
+		, EditLayerWeightmapAlpha
 		, bEditVisibility = bEditVisibility](FRDGBuilder& GraphBuilder)
 		{
 			TRACE_CPUPROFILER_EVENT_SCOPE(LandscapeCircleHeightPatch);
@@ -197,6 +207,7 @@ bool ULandscapeCircleHeightPatch::ApplyCirclePatch(UE::Landscape::EditLayers::FR
 			ShaderParams->InCenter = (FVector3f)CenterInHeightmapCoordinates;
 			ShaderParams->InRadius = HeightmapRadius;
 			ShaderParams->InFalloff = HeightmapFalloff;
+			ShaderParams->InEditLayerWeightmapAlpha = EditLayerWeightmapAlpha;
 			
 			ShaderParams->RenderTargets[0] = FRenderTargetBinding(DestinationTexture, ERenderTargetLoadAction::ENoAction,
 				/*InMipIndex = */0, LandscapeTextureSliceIndex);
diff --git a/Engine/Plugins/Editor/LandscapePatch/Source/LandscapePatch/Private/LandscapeCircleHeightPatchPS.h b/Engine/Plugins/Editor/LandscapePatch/Source/LandscapePatch/Private/LandscapeCircleHeightPatchPS.h
--- a/Engine/Plugins/Editor/LandscapePatch/Source/LandscapePatch/Private/LandscapeCircleHeightPatchPS.h
+++ b/Engine/Plugins/Editor/LandscapePatch/Source/LandscapePatch/Private/LandscapeCircleHeightPatchPS.h
@@ -18,6 +18,9 @@ public:
 		SHADER_PARAMETER(FVector3f, InCenter)
 		SHADER_PARAMETER(float, InRadius)
 		SHADER_PARAMETER(float, InFalloff)
+		// Texture patches apply the assigned edit layer's weightmap alpha to all weightmaps (including visibility)
+		// For consistency, the circle patch similarly applies the edit layer's weightmap alpha to visibility
+		SHADER_PARAMETER(float, InEditLayerWeightmapAlpha)
 		RENDER_TARGET_BINDING_SLOTS() // Holds our output
 	END_SHADER_PARAMETER_STRUCT()
 
diff --git a/Engine/Plugins/Editor/LandscapePatch/Source/LandscapePatch/Private/LandscapeTexturePatch.cpp b/Engine/Plugins/Editor/LandscapePatch/Source/LandscapePatch/Private/LandscapeTexturePatch.cpp
--- a/Engine/Plugins/Editor/LandscapePatch/Source/LandscapePatch/Private/LandscapeTexturePatch.cpp
+++ b/Engine/Plugins/Editor/LandscapePatch/Source/LandscapePatch/Private/LandscapeTexturePatch.cpp
@@ -854,6 +854,10 @@ void ULandscapeTexturePatch::GetHeightShaderParams(const FTransform& LandscapeHe
 	ParamsOut.InAlphaMode = static_cast<uint8>(HeightAlphaSettings.GetAlphaSourceMode());
 	ParamsOut.InAlphaTextureChannel = static_cast<uint8>(HeightAlphaSettings.AlphaTextureChannel);
 	ParamsOut.InAlphaTextureTransformToUV = static_cast<FMatrix44f>(HeightAlphaSettings.AlphaTextureUVTransform.GetTransformToUV(PatchToWorld, ParamsOut.InPatchWorldDimensions));
+
+	const ULandscapePatchEditLayer* PatchEditLayer = GetBoundEditLayer();
+	check(PatchEditLayer != nullptr);
+	ParamsOut.InEditLayerHeightmapAlpha = PatchEditLayer->GetAlphaForTargetType(ELandscapeToolTargetType::Heightmap);
 }
 
 void ULandscapeTexturePatch::GetWeightShaderParams(
@@ -885,6 +889,10 @@ void ULandscapeTexturePatch::GetWeightShaderParams(
 	ParamsOut.InAlphaMode = static_cast<uint8>(GetWeightPatchAlphaSourceMode(WeightPatchInfo->WeightmapLayerName));
 	ParamsOut.InAlphaTextureChannel = static_cast<uint8>(GetWeightPatchAlphaTextureChannel(WeightPatchInfo->WeightmapLayerName));
 	ParamsOut.InAlphaTextureTransformToUV = static_cast<FMatrix44f>(GetWeightPatchAlphaTextureTransform(WeightPatchInfo->WeightmapLayerName).GetTransformToUV(PatchToWorld, ParamsOut.InPatchWorldDimensions));
+
+	const ULandscapePatchEditLayer* PatchEditLayer = GetBoundEditLayer();
+	check(PatchEditLayer != nullptr);
+	ParamsOut.InEditLayerWeightmapAlpha = PatchEditLayer->GetAlphaForTargetType(ELandscapeToolTargetType::Weightmap);
 }
 
 // This function determines how our internal height render targets get converted to the format that gets
diff --git a/Engine/Plugins/Editor/LandscapePatch/Source/LandscapePatch/Private/LandscapeTexturePatchPS.h b/Engine/Plugins/Editor/LandscapePatch/Source/LandscapePatch/Private/LandscapeTexturePatchPS.h
--- a/Engine/Plugins/Editor/LandscapePatch/Source/LandscapePatch/Private/LandscapeTexturePatchPS.h
+++ b/Engine/Plugins/Editor/LandscapePatch/Source/LandscapePatch/Private/LandscapeTexturePatchPS.h
@@ -69,6 +69,8 @@ namespace UE::Landscape
 			SHADER_PARAMETER_RDG_TEXTURE_SRV(Texture2D<float4>, InHeightPatchAlphaTexture)
 			SHADER_PARAMETER(uint32, InAlphaTextureChannel)
 			SHADER_PARAMETER(FMatrix44f, InAlphaTextureTransformToUV)
+			// Assigned edit layer's heightmap alpha range [-1, 1]
+			SHADER_PARAMETER(float, InEditLayerHeightmapAlpha)
 
 			RENDER_TARGET_BINDING_SLOTS() // Holds our output
 		END_SHADER_PARAMETER_STRUCT()
@@ -203,6 +205,8 @@ namespace UE::Landscape
 			SHADER_PARAMETER_RDG_TEXTURE_SRV(Texture2D<float4>, InWeightPatchAlphaTexture)
 			SHADER_PARAMETER(uint32, InAlphaTextureChannel)
 			SHADER_PARAMETER(FMatrix44f, InAlphaTextureTransformToUV)
+			// Assigned edit layer's weightmap alpha range [0, 1]
+			SHADER_PARAMETER(float, InEditLayerWeightmapAlpha)
 
 			RENDER_TARGET_BINDING_SLOTS() // Holds our output
 		END_SHADER_PARAMETER_STRUCT()
diff --git a/Engine/Plugins/Editor/LandscapePatch/Source/LandscapePatch/Public/LandscapePatchComponent.h b/Engine/Plugins/Editor/LandscapePatch/Source/LandscapePatch/Public/LandscapePatchComponent.h
--- a/Engine/Plugins/Editor/LandscapePatch/Source/LandscapePatch/Public/LandscapePatchComponent.h
+++ b/Engine/Plugins/Editor/LandscapePatch/Source/LandscapePatch/Public/LandscapePatchComponent.h
@@ -174,7 +174,7 @@ public:
 protected:
 
 #if WITH_EDITOR
-	ULandscapePatchEditLayer* GetBoundEditLayer() { return EditLayer.Get(); }
+	ULandscapePatchEditLayer* GetBoundEditLayer() const { return EditLayer.Get(); }
 #endif
 
 	// Guid of the edit layer to which the patch is bound
