Landscape Patch - Added GetWeightPatch helper fuction
* Warns if patch with WeightmapLayerName does not exist or more than one exists within the same texture patch
* In the future, patches will warn the user on load and auto-fix fixup to avoid layer name conflicts

#Synced-CL 50538450

diff --git a/Engine/Plugins/Editor/LandscapePatch/Source/LandscapePatch/Private/LandscapeTexturePatch.cpp b/Engine/Plugins/Editor/LandscapePatch/Source/LandscapePatch/Private/LandscapeTexturePatch.cpp
--- a/Engine/Plugins/Editor/LandscapePatch/Source/LandscapePatch/Private/LandscapeTexturePatch.cpp
+++ b/Engine/Plugins/Editor/LandscapePatch/Source/LandscapePatch/Private/LandscapeTexturePatch.cpp
@@ -2488,16 +2488,11 @@ ELandscapeTexturePatchAlphaSourceMode ULandscapeTexturePatch::GetWeightPatchAlph
 		return WeightPatchesAlphaSettings.GetAlphaSourceMode();
 	}
 
-	for (const TObjectPtr<ULandscapeWeightPatchTextureInfo>& WeightPatch : WeightPatches)
+	if (TObjectPtr<ULandscapeWeightPatchTextureInfo> WeightPatch = GetWeightPatch(InWeightmapLayerName, ANSI_TO_TCHAR(__FUNCTION__)))
 	{
-		if (IsValid(WeightPatch) && WeightPatch->WeightmapLayerName == InWeightmapLayerName)
-		{
-			// Returns the override source mode if and only if the alpha settings are overridden
-			return GetWeightPatchAlphaSettingsConst(WeightPatch).GetAlphaSourceMode();
-		}
+		// Returns the override source mode if and only if the alpha settings are overridden
+		return GetWeightPatchAlphaSettingsConst(WeightPatch).GetAlphaSourceMode();
 	}
-
-	UE_LOG(LogLandscapePatch, Warning, TEXT("ULandscapeTexturePatch::GetWeightPatchAlphaSourceMode: Unable to get data for weight layer %s"), *InWeightmapLayerName.ToString());
 	return WeightPatchesAlphaSettings.GetAlphaSourceMode();
 }
 
@@ -2517,16 +2512,11 @@ void ULandscapeTexturePatch::SetWeightPatchAlphaSourceMode(const FName& InWeight
 	}
 
 	// Set the override weightmap source mode for a specific weight patch
-	for (const TObjectPtr<ULandscapeWeightPatchTextureInfo>& WeightPatch : WeightPatches)
+	if (TObjectPtr<ULandscapeWeightPatchTextureInfo> WeightPatch = GetWeightPatch(InWeightmapLayerName, ANSI_TO_TCHAR(__FUNCTION__)))
 	{
-		if (IsValid(WeightPatch) && WeightPatch->WeightmapLayerName == InWeightmapLayerName)
-		{
-			WeightPatch->SetWeightPatchAlphaSourceMode(InNewSourceMode, bInOverrideAlpha, bAlwaysMarkDirty);
-			return;
-		}
+		WeightPatch->SetWeightPatchAlphaSourceMode(InNewSourceMode, bInOverrideAlpha, bAlwaysMarkDirty);
+		return;
 	}
-	const FString LayerNameString = InWeightmapLayerName.ToString();
-	UE_LOG(LogLandscapePatch, Warning, TEXT("ULandscapeTexturePatch::SetWeightPatchAlphaSourceMode: Unable to find data for weight layer %s"), *LayerNameString);
 }
 
 const UTexture* ULandscapeTexturePatch::GetWeightPatchAlphaTextureAsset(const FName& InWeightmapLayerName) const
@@ -2537,16 +2527,12 @@ const UTexture* ULandscapeTexturePatch::GetWeightPatchAlphaTextureAsset(const FN
 		return WeightPatchesAlphaSettings.AlphaTextureAsset;
 	}
 
-	for (const TObjectPtr<ULandscapeWeightPatchTextureInfo>& WeightPatch : WeightPatches)
+	if (TObjectPtr<ULandscapeWeightPatchTextureInfo> WeightPatch = GetWeightPatch(InWeightmapLayerName, ANSI_TO_TCHAR(__FUNCTION__)))
 	{
-		if (IsValid(WeightPatch) && WeightPatch->WeightmapLayerName == InWeightmapLayerName)
-		{
-			// Returns the override value if and only if the alpha settings are overridden
-			return  GetWeightPatchAlphaSettingsConst(WeightPatch).AlphaTextureAsset;
-		}
+		// Returns the override value if and only if the alpha settings are overridden
+		return GetWeightPatchAlphaSettingsConst(WeightPatch).AlphaTextureAsset;
 	}
 
-	UE_LOG(LogLandscapePatch, Warning, TEXT("ULandscapeTexturePatch::GetWeightPatchAlphaTextureAsset: Unable to get data for weight layer %s"), *InWeightmapLayerName.ToString());
 	return WeightPatchesAlphaSettings.AlphaTextureAsset;
 }
 
@@ -2557,18 +2543,15 @@ UTextureRenderTarget2D* ULandscapeTexturePatch::GetWeightPatchAlphaRenderTarget(
 		return nullptr;
 	}
 
-	for (const TObjectPtr<ULandscapeWeightPatchTextureInfo>& WeightPatch : WeightPatches)
+	if (TObjectPtr<ULandscapeWeightPatchTextureInfo> WeightPatch = GetWeightPatch(InWeightmapLayerName, ANSI_TO_TCHAR(__FUNCTION__)))
 	{
-		if (IsValid(WeightPatch) && WeightPatch->WeightmapLayerName == InWeightmapLayerName)
+		if (bMarkDirty)
 		{
-			if (bMarkDirty)
-			{
-				MarkPackageDirty();
-			}
-
-			// returns the default (shared) weight patch render target or the specific override weight patch render target
-			return GetWeightPatchAlphaRenderTarget(WeightPatch, bMarkDirty);
+			MarkPackageDirty();
 		}
+
+		// returns the default (shared) weight patch render target or the specific override weight patch render target
+		return GetWeightPatchAlphaRenderTarget(WeightPatch, bMarkDirty);
 	}
 	return nullptr;
 }
@@ -2594,25 +2577,20 @@ void ULandscapeTexturePatch::SetWeightPatchAlphaTextureAsset(const FName& InWeig
 	}
 
 	// Set the override weightmap texture for a specific weight patch
-	for (const TObjectPtr<ULandscapeWeightPatchTextureInfo>& WeightPatch : WeightPatches)
+	if (TObjectPtr<ULandscapeWeightPatchTextureInfo> WeightPatch = GetWeightPatch(InWeightmapLayerName, ANSI_TO_TCHAR(__FUNCTION__)))
 	{
-		if (IsValid(WeightPatch) && WeightPatch->WeightmapLayerName == InWeightmapLayerName)
+		WeightPatch->bOverrideWeightAlphaSettings = bInOverrideAlpha;
+		// When bInOverrideAlpha is false the weightmap layer will fallback to the shared weightmap alpha
+		if (bInOverrideAlpha)
 		{
-			WeightPatch->bOverrideWeightAlphaSettings = bInOverrideAlpha;
-			// When bInOverrideAlpha is false the weightmap layer will fallback to the shared weightmap alpha
-			if (bInOverrideAlpha)
-			{
-				WeightPatch->WeightPatchAlphaSettings.AlphaTextureAsset = InTexture;
-			}
-			else
-			{
-				WeightPatch->SetWeightPatchAlphaSourceMode(ELandscapeTexturePatchAlphaSourceMode::None, /*bInOverrideAlpha =*/false, bAlwaysMarkDirty);
-			}
-			return;
+			WeightPatch->WeightPatchAlphaSettings.AlphaTextureAsset = InTexture;
+		}
+		else
+		{
+			WeightPatch->SetWeightPatchAlphaSourceMode(ELandscapeTexturePatchAlphaSourceMode::None, /*bInOverrideAlpha =*/false, bAlwaysMarkDirty);
 		}
+		return;
 	}
-	const FString LayerNameString = InWeightmapLayerName.ToString();
-	UE_LOG(LogLandscapePatch, Warning, TEXT("ULandscapeTexturePatch::SetWeightPatchAlphaTextureAsset: Unable to find data for weight layer %s"), *LayerNameString);
 }
 
 const FLandscapeTexturePatchTextureTransform& ULandscapeTexturePatch::GetWeightPatchAlphaTextureTransform(const FName& InWeightmapLayerName) const
@@ -2623,15 +2601,10 @@ const FLandscapeTexturePatchTextureTransform& ULandscapeTexturePatch::GetWeightP
 		return WeightPatchesAlphaSettings.AlphaTextureUVTransform;
 	}
 
-	for (const TObjectPtr<ULandscapeWeightPatchTextureInfo>& WeightPatch : WeightPatches)
+	if (TObjectPtr<ULandscapeWeightPatchTextureInfo> WeightPatch = GetWeightPatch(InWeightmapLayerName, ANSI_TO_TCHAR(__FUNCTION__)))
 	{
-		if (IsValid(WeightPatch) && WeightPatch->WeightmapLayerName == InWeightmapLayerName)
-		{
-			return GetWeightPatchAlphaSettingsConst(WeightPatch).AlphaTextureUVTransform;
-		}
+		return GetWeightPatchAlphaSettingsConst(WeightPatch).AlphaTextureUVTransform;
 	}
-
-	UE_LOG(LogLandscapePatch, Warning, TEXT("ULandscapeTexturePatch::GetWeightPatchAlphaTextureTransform: Unable to get data for weight layer %s"), *InWeightmapLayerName.ToString());
 	return WeightPatchesAlphaSettings.AlphaTextureUVTransform;
 }
 
@@ -2649,16 +2622,11 @@ void ULandscapeTexturePatch::SetWeightPatchAlphaTextureTransform(const FName& In
 		return;
 	}
 
-	for (const TObjectPtr<ULandscapeWeightPatchTextureInfo>& WeightPatch : WeightPatches)
+	if (TObjectPtr<ULandscapeWeightPatchTextureInfo> WeightPatch = GetWeightPatch(InWeightmapLayerName, ANSI_TO_TCHAR(__FUNCTION__)))
 	{
-		if (IsValid(WeightPatch) && WeightPatch->WeightmapLayerName == InWeightmapLayerName)
-		{
-			WeightPatch->SetWeightPatchAlphaTextureTransform(InTextureTransform, bAlwaysMarkDirty);
-			return;
-		}
+		WeightPatch->SetWeightPatchAlphaTextureTransform(InTextureTransform, bAlwaysMarkDirty);
+		return;
 	}
-
-	UE_LOG(LogLandscapePatch, Warning, TEXT("ULandscapeTexturePatch::SetWeightPatchAlphaTextureTransform: Unable to set data for weight layer %s"), *InWeightmapLayerName.ToString());
 }
 
 ELandscapeTexturePatchTextureChannel ULandscapeTexturePatch::GetWeightPatchAlphaTextureChannel(const FName& InWeightmapLayerName) const
@@ -2668,16 +2636,12 @@ ELandscapeTexturePatchTextureChannel ULandscapeTexturePatch::GetWeightPatchAlpha
 		return WeightPatchesAlphaSettings.AlphaTextureChannel;
 	}
 
-	for (const TObjectPtr<ULandscapeWeightPatchTextureInfo>& WeightPatch : WeightPatches)
+	if (TObjectPtr<ULandscapeWeightPatchTextureInfo> WeightPatch = GetWeightPatch(InWeightmapLayerName, ANSI_TO_TCHAR(__FUNCTION__)))
 	{
-		if (IsValid(WeightPatch) && WeightPatch->WeightmapLayerName == InWeightmapLayerName)
-		{
-			// Returns the override value if and only if the alpha settings are overridden
-			return GetWeightPatchAlphaSettingsConst(WeightPatch).AlphaTextureChannel;
-		}
+		// Returns the override value if and only if the alpha settings are overridden
+		return GetWeightPatchAlphaSettingsConst(WeightPatch).AlphaTextureChannel;
 	}
 
-	UE_LOG(LogLandscapePatch, Warning, TEXT("ULandscapeTexturePatch::GetWeightPatchAlphaTextureChannel: Unable to find data for weight layer %s"), *InWeightmapLayerName.ToString());
 	return WeightPatchesAlphaSettings.AlphaTextureChannel;
 }
 
@@ -2695,25 +2659,20 @@ void ULandscapeTexturePatch::SetWeightPatchAlphaTextureChannel(const FName& InWe
 		return;
 	}
 
-	for (const TObjectPtr<ULandscapeWeightPatchTextureInfo>& WeightPatch : WeightPatches)
+	if (TObjectPtr<ULandscapeWeightPatchTextureInfo> WeightPatch = GetWeightPatch(InWeightmapLayerName, ANSI_TO_TCHAR(__FUNCTION__)))
 	{
-		if (IsValid(WeightPatch) && WeightPatch->WeightmapLayerName == InWeightmapLayerName)
+		WeightPatch->bOverrideWeightAlphaSettings = bInOverrideAlpha;
+		// When bInOverrideAlpha is false the weightmap layer will fallback to the shared weightmap channel
+		if (bInOverrideAlpha)
 		{
-			WeightPatch->bOverrideWeightAlphaSettings = bInOverrideAlpha;
-			// When bInOverrideAlpha is false the weightmap layer will fallback to the shared weightmap channel
-			if (bInOverrideAlpha)
-			{
-				WeightPatch->SetWeightPatchAlphaTextureChannel(InTextureChannel, bAlwaysMarkDirty);
-			}
-			else
-			{
-				WeightPatch->SetWeightPatchAlphaSourceMode(ELandscapeTexturePatchAlphaSourceMode::None, /*bInOverrideAlpha =*/false, bAlwaysMarkDirty);
-			}
-			return;
+			WeightPatch->SetWeightPatchAlphaTextureChannel(InTextureChannel, bAlwaysMarkDirty);
 		}
+		else
+		{
+			WeightPatch->SetWeightPatchAlphaSourceMode(ELandscapeTexturePatchAlphaSourceMode::None, /*bInOverrideAlpha =*/false, bAlwaysMarkDirty);
+		}
+		return;
 	}
-
-	UE_LOG(LogLandscapePatch, Warning, TEXT("ULandscapeTexturePatch::SetWeightPatchAlphaTextureChannel: Unable to set data for weight layer %s"), *InWeightmapLayerName.ToString());
 }
 
 void ULandscapeTexturePatch::SetWeightPatchSourceMode(const FName& InWeightmapLayerName, ELandscapeTexturePatchSourceMode NewMode)
@@ -2721,39 +2680,28 @@ void ULandscapeTexturePatch::SetWeightPatchSourceMode(const FName& InWeightmapLa
 #if WITH_EDITOR
 	using namespace LandscapeTexturePatchLocals;
 
-	for (const TObjectPtr<ULandscapeWeightPatchTextureInfo>& WeightPatch : WeightPatches)
+	if (TObjectPtr<ULandscapeWeightPatchTextureInfo> WeightPatch = GetWeightPatch(InWeightmapLayerName, ANSI_TO_TCHAR(__FUNCTION__)))
 	{
-		if (IsValid(WeightPatch) && WeightPatch->WeightmapLayerName == InWeightmapLayerName)
-		{
-			WeightPatch->SetSourceMode(NewMode);
-			return;
-		}
+		WeightPatch->SetSourceMode(NewMode);
+		return;
 	}
-	const FString LayerNameString = InWeightmapLayerName.ToString();
-	UE_LOG(LogLandscapePatch, Warning, TEXT("ULandscapeTexturePatch::SetWeightPatchSourceMode: Unable to find data for weight layer %s"), *LayerNameString);
 #endif // WITH_EDITOR
 }
 
 ELandscapeTexturePatchSourceMode ULandscapeTexturePatch::GetWeightPatchSourceMode(const FName& InWeightmapLayerName)
 {
-	for (const TObjectPtr<ULandscapeWeightPatchTextureInfo>& WeightPatch : WeightPatches)
+	if (TObjectPtr<ULandscapeWeightPatchTextureInfo> WeightPatch = GetWeightPatch(InWeightmapLayerName, ANSI_TO_TCHAR(__FUNCTION__)))
 	{
-		if (IsValid(WeightPatch) && WeightPatch->WeightmapLayerName == InWeightmapLayerName)
-		{
-			return WeightPatch->SourceMode;
-		}
+		return WeightPatch->SourceMode;
 	}
 	return ELandscapeTexturePatchSourceMode::None;
 }
 
 const FLandscapeTexturePatchTextureTransform& ULandscapeTexturePatch::GetWeightPatchTextureTransform(const FName& InWeightmapLayerName) const
 {
-	for (const TObjectPtr<ULandscapeWeightPatchTextureInfo>& WeightPatch : WeightPatches)
+	if (TObjectPtr<ULandscapeWeightPatchTextureInfo> WeightPatch = GetWeightPatch(InWeightmapLayerName, ANSI_TO_TCHAR(__FUNCTION__)))
 	{
-		if (IsValid(WeightPatch) && WeightPatch->WeightmapLayerName == InWeightmapLayerName)
-		{
-			return WeightPatch->WeightTextureUVTransform;
-		}
+		return WeightPatch->WeightTextureUVTransform;
 	}
 
 	// There is no default texture transform for all weight patches (Only the alpha texture is overridable). Log an explicit error here
@@ -2763,53 +2711,36 @@ const FLandscapeTexturePatchTextureTransform& ULandscapeTexturePatch::GetWeightP
 
 void ULandscapeTexturePatch::SetWeightPatchTextureTransform(const FName& InWeightmapLayerName, const FLandscapeTexturePatchTextureTransform& InTextureTransform, bool bAlwaysMarkDirty)
 {
-	for (const TObjectPtr<ULandscapeWeightPatchTextureInfo>& WeightPatch : WeightPatches)
+	if (TObjectPtr<ULandscapeWeightPatchTextureInfo> WeightPatch = GetWeightPatch(InWeightmapLayerName, ANSI_TO_TCHAR(__FUNCTION__)))
 	{
-		if (IsValid(WeightPatch) && WeightPatch->WeightmapLayerName == InWeightmapLayerName)
-		{
-			WeightPatch->SetWeightPatchTextureTransform(InTextureTransform, bAlwaysMarkDirty);
-			return;
-		}
+		WeightPatch->SetWeightPatchTextureTransform(InTextureTransform, bAlwaysMarkDirty);
+		return;
 	}
-
-	UE_LOG(LogLandscapePatch, Warning, TEXT("ULandscapeTexturePatch::SetWeightPatchTextureTransform: Unable to set data for weight layer %s"), *InWeightmapLayerName.ToString());
 }
 
 ELandscapeTexturePatchTextureChannel ULandscapeTexturePatch::GetWeightPatchTextureChannel(const FName& InWeightmapLayerName) const
 {
-	for (const TObjectPtr<ULandscapeWeightPatchTextureInfo>& WeightPatch : WeightPatches)
+	if (TObjectPtr<ULandscapeWeightPatchTextureInfo> WeightPatch = GetWeightPatch(InWeightmapLayerName, ANSI_TO_TCHAR(__FUNCTION__)))
 	{
-		if (IsValid(WeightPatch) && WeightPatch->WeightmapLayerName == InWeightmapLayerName)
-		{
-			return WeightPatch->WeightPatchTextureChannel;
-		}
+		return WeightPatch->WeightPatchTextureChannel;
 	}
-	UE_LOG(LogLandscapePatch, Warning, TEXT("ULandscapeTexturePatch::GetWeightPatchTextureChannel: Unable to get data for weight layer %s"), *InWeightmapLayerName.ToString());
 	return ELandscapeTexturePatchTextureChannel::None;
 }
 
 void ULandscapeTexturePatch::SetWeightPatchTextureChannel(const FName& InWeightmapLayerName, ELandscapeTexturePatchTextureChannel InTextureChannel, bool bAlwaysMarkDirty)
 {
-	for (const TObjectPtr<ULandscapeWeightPatchTextureInfo>& WeightPatch : WeightPatches)
+	if (TObjectPtr<ULandscapeWeightPatchTextureInfo> WeightPatch = GetWeightPatch(InWeightmapLayerName, ANSI_TO_TCHAR(__FUNCTION__)))
 	{
-		if (IsValid(WeightPatch) && WeightPatch->WeightmapLayerName == InWeightmapLayerName)
-		{
-			WeightPatch->SetWeightPatchTextureChannel(InTextureChannel, bAlwaysMarkDirty);
-			return;
-		}
+		WeightPatch->SetWeightPatchTextureChannel(InTextureChannel, bAlwaysMarkDirty);
+		return;
 	}
-
-	UE_LOG(LogLandscapePatch, Warning, TEXT("ULandscapeTexturePatch::SetWeightPatchTextureChannel: Unable to set data for weight layer %s"), *InWeightmapLayerName.ToString());
 }
 
 UTexture* ULandscapeTexturePatch::GetWeightPatchTextureAsset(const FName& InWeightmapLayerName)
 {
-	for (const TObjectPtr<ULandscapeWeightPatchTextureInfo>& WeightPatch : WeightPatches)
+	if (TObjectPtr<ULandscapeWeightPatchTextureInfo> WeightPatch = GetWeightPatch(InWeightmapLayerName, ANSI_TO_TCHAR(__FUNCTION__)))
 	{
-		if (IsValid(WeightPatch) && WeightPatch->WeightmapLayerName == InWeightmapLayerName)
-		{
-			return WeightPatch->TextureAsset;
-		}
+		return WeightPatch->TextureAsset;
 	}
 	return nullptr;
 }
@@ -2821,17 +2752,14 @@ UTextureRenderTarget2D* ULandscapeTexturePatch::GetWeightPatchRenderTarget(const
 		return nullptr;
 	}
 
-	for (const TObjectPtr<ULandscapeWeightPatchTextureInfo>& WeightPatch : WeightPatches)
+	if (TObjectPtr<ULandscapeWeightPatchTextureInfo> WeightPatch = GetWeightPatch(InWeightmapLayerName, ANSI_TO_TCHAR(__FUNCTION__)))
 	{
-		if (IsValid(WeightPatch) && WeightPatch->WeightmapLayerName == InWeightmapLayerName)
+		if (bMarkDirty)
 		{
-			if (bMarkDirty)
-			{
-				MarkPackageDirty();
-			}
-
-			return GetWeightPatchRenderTarget(WeightPatch);
+			MarkPackageDirty();
 		}
+
+		return GetWeightPatchRenderTarget(WeightPatch);
 	}
 	return nullptr;
 }
@@ -2877,6 +2805,31 @@ const FLandscapeTexturePatchAlphaSettings& ULandscapeTexturePatch::GetWeightPatc
 	return WeightPatch->bOverrideWeightAlphaSettings ? WeightPatch->WeightPatchAlphaSettings : WeightPatchesAlphaSettings;
 }
 
+
+TObjectPtr<ULandscapeWeightPatchTextureInfo> ULandscapeTexturePatch::GetWeightPatch(const FName& InWeightmapLayerName, const TCHAR* InDebugLogName) const
+{
+	TArray<TObjectPtr<ULandscapeWeightPatchTextureInfo>> FoundPatches = WeightPatches.FilterByPredicate([&InWeightmapLayerName](const TObjectPtr<ULandscapeWeightPatchTextureInfo>& WeightPatch)
+		{
+			return IsValid(WeightPatch) && WeightPatch->WeightmapLayerName == InWeightmapLayerName;
+		});
+
+	const FString LayerNameString = InWeightmapLayerName.ToString();
+	const TCHAR* DebugLogName = (InDebugLogName == nullptr) ? TEXT("ULandscapeTexturePatch::GetWeightPatch") : InDebugLogName;
+
+	if (FoundPatches.IsEmpty())
+	{
+		UE_LOG(LogLandscapePatch, Warning, TEXT("%s: Unable to find weightmap patch for weightmap layer name %s"), DebugLogName, *LayerNameString);
+		return nullptr;
+	}
+
+	if (FoundPatches.Num() > 1)
+	{
+		UE_LOG(LogLandscapePatch, Warning, TEXT("%s found more than one patch using the same weightmap layer name %s. BP functions will return the first matching weight patch only!"), DebugLogName, *LayerNameString);
+	}
+
+	return FoundPatches[0];
+}
+
 UTextureRenderTarget2D* ULandscapeTexturePatch::GetWeightPatchAlphaRenderTarget(ULandscapeWeightPatchTextureInfo* WeightPatch, bool bMarkDirty)
 {
 	if (IsTemplate() || !IsValid(WeightPatch))
@@ -2943,13 +2896,10 @@ void ULandscapeTexturePatch::SetWeightPatchTextureAsset(const FName& InWeightmap
 		return;
 	}
 
-	for (const TObjectPtr<ULandscapeWeightPatchTextureInfo>& WeightPatch : WeightPatches)
+	if (TObjectPtr<ULandscapeWeightPatchTextureInfo> WeightPatch = GetWeightPatch(InWeightmapLayerName, ANSI_TO_TCHAR(__FUNCTION__)))
 	{
-		if (IsValid(WeightPatch) && WeightPatch->WeightmapLayerName == InWeightmapLayerName)
-		{
-			WeightPatch->TextureAsset = TextureIn;
-			return;
-		}
+		WeightPatch->TextureAsset = TextureIn;
+		return;
 	}
 
 	const FString LayerNameString = InWeightmapLayerName.ToString();
@@ -2958,37 +2908,26 @@ void ULandscapeTexturePatch::SetWeightPatchTextureAsset(const FName& InWeightmap
 
 void ULandscapeTexturePatch::SetWeightPatchBlendModeOverride(const FName& InWeightmapLayerName, ELandscapeTexturePatchBlendMode BlendModeIn)
 {
-	for (const TObjectPtr<ULandscapeWeightPatchTextureInfo>& WeightPatch : WeightPatches)
+	if (TObjectPtr<ULandscapeWeightPatchTextureInfo> WeightPatch = GetWeightPatch(InWeightmapLayerName, ANSI_TO_TCHAR(__FUNCTION__)))
 	{
-		if (IsValid(WeightPatch) && WeightPatch->WeightmapLayerName == InWeightmapLayerName)
-		{
-			WeightPatch->OverrideBlendMode = BlendModeIn;
-			WeightPatch->bOverrideBlendMode = true;
-			return;
-		}
+		WeightPatch->OverrideBlendMode = BlendModeIn;
+		WeightPatch->bOverrideBlendMode = true;
 	}
 }
 
 void ULandscapeTexturePatch::ClearWeightPatchBlendModeOverride(const FName& InWeightmapLayerName)
 {
-	for (const TObjectPtr<ULandscapeWeightPatchTextureInfo>& WeightPatch : WeightPatches)
+	if (TObjectPtr<ULandscapeWeightPatchTextureInfo> WeightPatch = GetWeightPatch(InWeightmapLayerName, ANSI_TO_TCHAR(__FUNCTION__)))
 	{
-		if (IsValid(WeightPatch) && WeightPatch->WeightmapLayerName == InWeightmapLayerName)
-		{
-			WeightPatch->bOverrideBlendMode = false;
-			return;
-		}
+		WeightPatch->bOverrideBlendMode = false;
 	}
 }
 
 void ULandscapeTexturePatch::SetEditVisibilityLayer(const FName& InWeightmapLayerName, const bool bEditVisibilityLayer)
 {
-	for (const TObjectPtr<ULandscapeWeightPatchTextureInfo>& WeightPatch : WeightPatches)
+	if (TObjectPtr<ULandscapeWeightPatchTextureInfo> WeightPatch = GetWeightPatch(InWeightmapLayerName, ANSI_TO_TCHAR(__FUNCTION__)))
 	{
-		if (IsValid(WeightPatch) && WeightPatch->WeightmapLayerName == InWeightmapLayerName)
-		{
-			WeightPatch->bEditVisibilityLayer = bEditVisibilityLayer;
-		}
+		WeightPatch->bEditVisibilityLayer = bEditVisibilityLayer;
 	}
 }
 
diff --git a/Engine/Plugins/Editor/LandscapePatch/Source/LandscapePatch/Public/LandscapeTexturePatch.h b/Engine/Plugins/Editor/LandscapePatch/Source/LandscapePatch/Public/LandscapeTexturePatch.h
--- a/Engine/Plugins/Editor/LandscapePatch/Source/LandscapePatch/Public/LandscapeTexturePatch.h
+++ b/Engine/Plugins/Editor/LandscapePatch/Source/LandscapePatch/Public/LandscapeTexturePatch.h
@@ -1021,6 +1021,8 @@ private:
 	FLandscapeTexturePatchAlphaSettings& GetWeightPatchAlphaSettings(ULandscapeWeightPatchTextureInfo* WeightPatch);
 	const FLandscapeTexturePatchAlphaSettings& GetWeightPatchAlphaSettingsConst(ULandscapeWeightPatchTextureInfo* WeightPatch) const; 
 
+	TObjectPtr<ULandscapeWeightPatchTextureInfo> GetWeightPatch(const FName& InWeightmapLayerName, const TCHAR* InDebugLogName = TEXT("ULandscapeTexturePatch::GetWeightPatch")) const;
+
 #if WITH_EDITOR
 	void TransitionHeightSourceModeInternal(ELandscapeTexturePatchSourceMode OldMode, ELandscapeTexturePatchSourceMode NewMode);
 	FLandscapeHeightPatchConvertToNativeParams GetHeightConvertToNativeParams() const;
