Landscape Texture Patch - Added UseWorldPositionSampling flag to allow patches to change texture sampling based on their world position.
* When world position sampling is enabled, texture is cenetered around the world origin and ignores patch/landscape transforms. Tiling the source/alpha texture tiles the texture from the world origin.

#Synced-CL 51092926

diff --git a/Engine/Plugins/Editor/LandscapePatch/Shaders/Private/LandscapeTexturePatchPS.usf b/Engine/Plugins/Editor/LandscapePatch/Shaders/Private/LandscapeTexturePatchPS.usf
--- a/Engine/Plugins/Editor/LandscapePatch/Shaders/Private/LandscapeTexturePatchPS.usf
+++ b/Engine/Plugins/Editor/LandscapePatch/Shaders/Private/LandscapeTexturePatchPS.usf
@@ -79,12 +79,6 @@ float GetFalloffAlpha(float FalloffWorldMargin, float2 PatchWorldDimensions, flo
 	return Alpha;
 }
 
-// Applies UV space transform (scale, rotation, offset) to a given UV
-float2 GetTransformedUV(float2 InUV, float4x4 InTransform)
-{
-	return mul(float4(InUV.xy, 0.0, 1.0), InTransform).xy;	
-}
-
 float GetChannelFromSample(uint InTextureChannel, float4 InSampledValue)
 {
 	switch (InTextureChannel)
@@ -125,8 +119,10 @@ float GetAlphaFactor(uint InAlphaMode, uint InAlphaTextureChannel, float4 InPatc
 #if APPLY_HEIGHT_PATCH
 Texture2D<float4> InHeightPatch;
 SamplerState InHeightPatchSampler;
-float4x4 InHeightmapToPatch;
-float4x4 InHeightmapTextureTransformToUV;
+// Heightmap to PatchUV transform. Default patch behavior that samples the texture relative to the patch's bounds
+float4x4 InHeightmapToPatchUV;
+// InHeightmapToPatchUV matrix with the source texture transform (world space or relative) applied
+float4x4 InHeightmapToPatchUVTransform;
 float2 InPatchWorldDimensions;
 float2 InEdgeUVDeadBorder;
 float InFalloffWorldMargin;
@@ -146,7 +142,8 @@ uint InTextureChannel;
 uint InAlphaMode;
 Texture2D<float4> InHeightPatchAlphaTexture;
 uint InAlphaTextureChannel;
-float4x4 InAlphaTextureTransformToUV;
+// InHeightmapToPatchUV matrix with the alpha source texture transform (world space or relative) applied
+float4x4 InHeightmapAlphaToPatchUVTransform;
 // Assigned edit layer's heightmap/weightmap alpha
 float InEditLayerHeightmapAlpha;
 float InEditLayerWeightmapAlpha;
@@ -174,17 +171,16 @@ void ApplyLandscapeTextureHeightPatch(in float4 SVPos : SV_POSITION, out float4
 	// so that they match to the cpp file
 	bool bRectangularFalloff = InFlags & RECTANGULAR_FALLOFF_FLAG;
 	bool bInputIsPackedHeight = InFlags & INPUT_IS_PACKED_HEIGHT_FLAG;
-	
+
 	// We need only the 2D affine transformation that goes from landscape XY heightmap integer coordinates
-	// to patch UV coordinates.
-	float2x2 HeightmapToPatchRotateScale = (float2x2) InHeightmapToPatch;
-	float2 HeightmapToPatchTranslate = InHeightmapToPatch._m03_m13;
-	
+	// to patch UV coordinates. SVPos is in batched merge texture coordinates here
+	float4 SVPosCoordinates = float4(SVPos.xy, 0.0, 1.0);
 	// Store normalized coordinates that map the patch to a [0, 1] UV space with no tiling (used for patch falloff)
-	float2 PatchUVCoordinates = mul(HeightmapToPatchRotateScale, SVPos.xy) + HeightmapToPatchTranslate;
-	// Apply the scale, offset, rotation UV space transform to convert from normalized UV to tiling space (utilizes AM_Wrap sampler) 
-	float2 PatchUVCoordinatesTransformed = GetTransformedUV(PatchUVCoordinates, InHeightmapTextureTransformToUV);	
-	float2 PatchAlphaUVCoordinatesTransformed = GetTransformedUV(PatchUVCoordinates, InAlphaTextureTransformToUV);	
+	float4 PatchUVCoordinates = float4(mul(InHeightmapToPatchUV, SVPosCoordinates).xy, 0, 1);
+
+	// Apply the HeightmapToPatch matrix with scale, offset, and rotation to convert from normalized UV to tiling space (utilizes AM_Wrap sampler)
+	float2 PatchUVCoordinatesTransformed = mul(InHeightmapToPatchUVTransform, SVPosCoordinates).xy;
+	float2 PatchAlphaUVCoordinatesTransformed = mul(InHeightmapAlphaToPatchUVTransform, SVPosCoordinates).xy;
 
 	float4 PatchSampledValue = InHeightPatch.Sample(InHeightPatchSampler, PatchUVCoordinatesTransformed);
 	// When using an internal texture/native packed height we enforce that the data comes from the RG channels
@@ -193,7 +189,7 @@ void ApplyLandscapeTextureHeightPatch(in float4 SVPos : SV_POSITION, out float4
 	
 	float PatchSignedHeight = InHeightScale * (PatchStoredHeight - InZeroInEncoding) + InHeightOffset;
 	
-	float Alpha = GetFalloffAlpha(InFalloffWorldMargin, InPatchWorldDimensions, PatchUVCoordinates, InEdgeUVDeadBorder, bRectangularFalloff);
+	float Alpha = GetFalloffAlpha(InFalloffWorldMargin, InPatchWorldDimensions, PatchUVCoordinates.xy, InEdgeUVDeadBorder, bRectangularFalloff);
 	
 	float4 PatchAlphaSampledValue = InHeightPatchAlphaTexture.Sample(InHeightPatchSampler, PatchAlphaUVCoordinatesTransformed);
 	// Apply the alpha from the source height texture or use the standalone alpha texture
@@ -252,11 +248,12 @@ void ConvertBackFromNativeLandscapePatch(in float4 SVPos : SV_POSITION, out floa
 
 #if APPLY_WEIGHT_PATCH
 Texture2D<float4> InSourceWeightmap;
-uint2 InSourceWeightmapCoordOffset;
 Texture2D<float4> InWeightPatch;
 SamplerState InWeightPatchSampler;
-float4x4 InWeightmapToPatch;
-float4x4 InWeightmapTextureTransformToUV;
+// Weightmap to PatchUV transform. Default patch behavior that samples the texture relative to the patch's bounds
+float4x4 InWeightmapToPatchUV;
+// InWeightmapToPatchUV matrix with the source texture transform (world space or relative) applied
+float4x4 InWeightmapToPatchUVTransform;
 float2 InPatchWorldDimensions;
 float2 InEdgeUVDeadBorder;
 float InFalloffWorldMargin;
@@ -270,7 +267,8 @@ uint InTextureChannel;
 uint InAlphaMode;
 Texture2D<float4> InWeightPatchAlphaTexture;
 uint InAlphaTextureChannel;
-float4x4 InAlphaTextureTransformToUV;
+// InWeightmapToPatchUV matrix with the alpha source texture transform (world space or relative) applied
+float4x4 InWeightmapAlphaToPatchUVTransform;
 // Patch edit layer's weightmap alpha
 float InEditLayerWeightmapAlpha;
 
@@ -298,21 +296,20 @@ void ApplyLandscapeTextureWeightPatch(in float4 SVPos : SV_POSITION, out float4
 	bool bRectangularFalloff = InFlags & RECTANGULAR_FALLOFF_FLAG;
 
 	// We need only the 2D affine transformation that goes from landscape weightmap integer coordinates
-	// to patch UV coordinates.
-	float2x2 WeightmapToPatchRotateScale = (float2x2) InWeightmapToPatch;
-	float2 WeightmapToPatchTranslate = InWeightmapToPatch._m03_m13;
-	
+	// to patch UV coordinates. SVPos is in batched merge texture coordinates here
+	float4 SVPosCoordinates = float4(SVPos.xy, 0.0, 1.0);
 	// Store normalized coordinates that map the patch to a [0, 1] UV space with no tiling (used for patch falloff)
-	float2 PatchUVCoordinates = mul(WeightmapToPatchRotateScale, SVPos.xy) + WeightmapToPatchTranslate;
-	// Apply the scale, offset, rotation UV space transform to convert from normalized UV to tiling space (utilizes AM_Wrap sampler) 
-	float2 PatchUVCoordinatesTransformed = GetTransformedUV(PatchUVCoordinates, InWeightmapTextureTransformToUV);
-	float2 PatchAlphaUVCoordinatesTransformed = GetTransformedUV(PatchUVCoordinates, InAlphaTextureTransformToUV);
+	float4 PatchUVCoordinates = float4(mul(InWeightmapToPatchUV, SVPosCoordinates).xy, 0, 1);
+
+	// Apply the WeightmapToPatch matrix with scale, offset, and rotation to convert from normalized UV to tiling space (utilizes AM_Wrap sampler)
+	float2 PatchUVCoordinatesTransformed = mul(InWeightmapToPatchUVTransform, SVPosCoordinates).xy;
+	float2 PatchAlphaUVCoordinatesTransformed = mul(InWeightmapAlphaToPatchUVTransform, SVPosCoordinates).xy;
 
 	float4 PatchSampledValue = InWeightPatch.Sample(InWeightPatchSampler, PatchUVCoordinatesTransformed);
 	// Get the selected texture color channel from the sample
 	float PatchWeight = GetChannelFromSample(InTextureChannel, PatchSampledValue);
 
-	float Alpha = GetFalloffAlpha(InFalloffWorldMargin, InPatchWorldDimensions, PatchUVCoordinates, InEdgeUVDeadBorder, bRectangularFalloff);
+	float Alpha = GetFalloffAlpha(InFalloffWorldMargin, InPatchWorldDimensions, PatchUVCoordinates.xy, InEdgeUVDeadBorder, bRectangularFalloff);
 
 	// Apply the alpha from the source weight texture or use the standalone alpha texture
 	float4 PatchAlphaSampledValue = InWeightPatchAlphaTexture.Sample(InWeightPatchSampler, PatchAlphaUVCoordinatesTransformed);
diff --git a/Engine/Plugins/Editor/LandscapePatch/Source/LandscapePatch/Private/LandscapeTexturePatch.cpp b/Engine/Plugins/Editor/LandscapePatch/Source/LandscapePatch/Private/LandscapeTexturePatch.cpp
--- a/Engine/Plugins/Editor/LandscapePatch/Source/LandscapePatch/Private/LandscapeTexturePatch.cpp
+++ b/Engine/Plugins/Editor/LandscapePatch/Source/LandscapePatch/Private/LandscapeTexturePatch.cpp
@@ -195,20 +195,46 @@ namespace LandscapeTexturePatchLocals
 	}
 }
 
-FMatrix44d FLandscapeTexturePatchTextureTransform::GetTransformToUV(const FTransform& PatchToWorldIn, const FVector2f& PatchWorldDimensionsIn) const
+FMatrix44d FLandscapeTexturePatchTextureTransform::GetTransformAsMatrix(const FTransform& PatchToWorldIn, const FVector2f& PatchWorldDimensionsIn) const
 {
-	// Calculate UV transforms: Scale (relative or world), Offset, and Rotation. Shader applies transformation to patch UV before sampling
-	// Scale is relative to local UV space unless bUseWorldSpaceScale is enabled. PatchWorldDimensions = PatchUnscaledSize * PatchScale
-	FVector UVOffset = FVector(Offset.X, Offset.Y, 0.0);
-	
+	const FVector PatchWorldScale = PatchToWorldIn.GetScale3D();
+	const FVector2f ScaledPatchDimensions = PatchWorldDimensionsIn * FVector2f(PatchWorldScale.X, PatchWorldScale.Y);
+
 	// Clamp scale to avoid divide by zero
 	const double ScaleX = FMath::Abs(Scale.X) < SMALL_NUMBER ? FMath::FloatSelect(Scale.X, 1.0, -1.0) * SMALL_NUMBER : Scale.X;
 	const double ScaleY = FMath::Abs(Scale.Y) < SMALL_NUMBER ? FMath::FloatSelect(Scale.Y, 1.0, -1.0) * SMALL_NUMBER : Scale.Y;
-	FVector UVRelativeScale = FVector(1.0 / ScaleX, 1.0 / ScaleY, 1.0);
-	const FVector PatchWorldScale = PatchToWorldIn.GetScale3D();
-	FVector2f UVWorldScale = (PatchWorldDimensionsIn * FVector2f(PatchWorldScale.X, PatchWorldScale.Y)) / FVector2f(ScaleX, ScaleY);
 
-	return FTransform(FRotator(0.0, Rotation, 0.0), UVOffset, bUseWorldSpaceScale ? FVector(UVWorldScale.X, UVWorldScale.Y, 1.0) : UVRelativeScale).ToMatrixWithScale();
+	// Offset may be in [0, 1] range or in world units
+	FVector UVOffset = FVector(Offset.X, Offset.Y, 0.0);
+
+	// Both world/relative sampling support local and world scaling factors
+	if (bUseWorldPositionSampling)
+	{
+		// Return WorldtoTexture matrix
+		const FVector UVWorldScale = FVector(ScaleX, ScaleY, 1.0);
+		const FVector UVRelativeScale = FVector(ScaledPatchDimensions.X * ScaleX, ScaledPatchDimensions.Y * ScaleY, 1.0);
+
+		// Add half the texture UV tiling scale to center the texture around origin. Scale will be in [0, 1] range or world units
+		UVOffset += FVector(ScaleX / 2, ScaleY / 2, 0);
+
+		// By default the transform is in world units. When using relative scale, multiply by the patch dimensions to get the offset in relative UV space
+		if (!bUseWorldSpaceScale)
+		{
+			UVOffset *= FVector(ScaledPatchDimensions.X, ScaledPatchDimensions.Y, 0);
+		}
+
+		const FTransform TextureUVToWorld = FTransform(FRotator(0, Rotation, 0), UVOffset, bUseWorldSpaceScale ? UVWorldScale : UVRelativeScale);
+		// Return a WorldToTextureUV matrix
+		return TextureUVToWorld.ToInverseMatrixWithScale();
+	}
+	else
+	{
+		// This matrix is used to apply a rotation, offset, scale to the patch's UVs
+		FVector UVRelativeScale = FVector(1.0 / ScaleX, 1.0 / ScaleY, 1.0);
+		FVector2f UVWorldScale = ScaledPatchDimensions / FVector2f(ScaleX, ScaleY);
+
+		return FTransform(FRotator(0.0, Rotation, 0.0), UVOffset, bUseWorldSpaceScale ? FVector(UVWorldScale.X, UVWorldScale.Y, 1.0) : UVRelativeScale).ToMatrixWithScale();
+	}
 }
 
 UTextureRenderTarget2D* FLandscapeTexturePatchAlphaSettings::GetAlphaRenderTarget(UObject* InOuter, bool bAlwaysMarkDirty/*=true*/)
@@ -725,9 +751,10 @@ void ULandscapeTexturePatch::ApplyToWeightmap(UE::Landscape::EditLayers::FRender
 }
 
 void ULandscapeTexturePatch::GetCommonShaderParams(const FTransform& LandscapeHeightmapToWorldIn,
-	const FIntPoint& SourceResolutionIn, const FIntPoint& DestinationResolutionIn, const FLandscapeTexturePatchTextureTransform& SourceTextureUVTransformIn,
-	FTransform& PatchToWorldOut, FVector2f& PatchWorldDimensionsOut, FMatrix44f& HeightmapToPatchOut, FMatrix44f& SourceTextureTransformToUVOut, FIntRect& DestinationBoundsOut,
-	FVector2f& EdgeUVDeadBorderOut, float& FalloffWorldMarginOut) const
+	const FIntPoint& SourceResolutionIn, const FIntPoint& DestinationResolutionIn,
+	const FLandscapeTexturePatchTextureTransform& SourceTextureTransformIn, const FLandscapeTexturePatchTextureTransform& AlphaTextureTransformIn,
+	FTransform& PatchToWorldOut, FVector2f& PatchWorldDimensionsOut, FMatrix44f& HeightmapToPatchUVOut, FMatrix44f& HeightmapToPatchUVTransformOut,
+	FMatrix44f& HeightmapAlphaToPatchUVTransformOut, FIntRect& DestinationBoundsOut, FVector2f& EdgeUVDeadBorderOut, float& FalloffWorldMarginOut) const
 {
 	using namespace LandscapeTexturePatchLocals;
 
@@ -748,10 +775,39 @@ void ULandscapeTexturePatch::GetCommonShaderParams(const FTransform& LandscapeHe
 	// In unreal, matrix composition is done by multiplying the subsequent ones on the right, and the result
 	// is transpose of what our shader will expect (because unreal right multiplies vectors by matrices).
 	FMatrix44d LandscapeToPatchUVTransposed = LandscapeToWorld * WorldToPatch * PatchLocalToUVs;
-	HeightmapToPatchOut = static_cast<FMatrix44f>(LandscapeToPatchUVTransposed.GetTransposed());
+	HeightmapToPatchUVOut = static_cast<FMatrix44f>(LandscapeToPatchUVTransposed.GetTransposed());
+
+	// The shader uses HeightmapToPatchUV matrix to map SVPos (batched merge texture space) to the [0, 1] UV space. This is useful for applying falloff along the patch edges
+	// The transformed matrices apply the source texture's transform (Scale, Offset, Rotation) to use custom texture sampling within the patch bounds (UV tiling)
+	if (SourceTextureTransformIn.bUseWorldPositionSampling)
+	{
+		// World space matrix transform : LandscapeHeightmapToWorldIn * WorldToTexture
+		const FMatrix44d WorldToTextureUV = SourceTextureTransformIn.GetTransformAsMatrix(PatchToWorldOut, PatchWorldDimensionsOut);
+		const FMatrix44d HeightmapToTextureUV = LandscapeToWorld * WorldToTextureUV;
+		HeightmapToPatchUVTransformOut = static_cast<FMatrix44f>(HeightmapToTextureUV.GetTransposed());
+	}
+	else
+	{
+		// Relative space matrix transform : LandscapeToPatchUVTransposed is the default PatchUV space (based on patch's bounds) 
+		const FMatrix44d LandscapeToTransformPatchUV = LandscapeToPatchUVTransposed * SourceTextureTransformIn.GetTransformAsMatrix(PatchToWorldOut, PatchWorldDimensionsOut);
+		HeightmapToPatchUVTransformOut = static_cast<FMatrix44f>(LandscapeToTransformPatchUV.GetTransposed());
+	}
+
+	// Build the alpha texture transformed matrix
+	if (AlphaTextureTransformIn.bUseWorldPositionSampling)
+	{
+		// World space matrix transform : LandscapeHeightmapToWorldIn * WorldToTexture
+		const FMatrix44d WorldToAlphaTextureUV = AlphaTextureTransformIn.GetTransformAsMatrix(PatchToWorldOut, PatchWorldDimensionsOut);
+		const FMatrix44d HeightmapAlphaToTextureUV = LandscapeToWorld * WorldToAlphaTextureUV;
+		HeightmapAlphaToPatchUVTransformOut = static_cast<FMatrix44f>(HeightmapAlphaToTextureUV.GetTransposed());
+	}
+	else
+	{
+		// Relative space matrix transform : LandscapeToPatchUVTransposed is the default PatchUV space (based on patch's bounds) 
+		const FMatrix44d LandscapeToTransformPatchAlphaUV = LandscapeToPatchUVTransposed * AlphaTextureTransformIn.GetTransformAsMatrix(PatchToWorldOut, PatchWorldDimensionsOut);
+		HeightmapAlphaToPatchUVTransformOut = static_cast<FMatrix44f>(LandscapeToTransformPatchAlphaUV.GetTransposed());
+	}
 
-	SourceTextureTransformToUVOut = static_cast<FMatrix44f>(SourceTextureUVTransformIn.GetTransformToUV(PatchToWorldOut, PatchWorldDimensionsOut));
-	
 	// Get the output bounds, which are used to limit the amount of landscape pixels we have to process. 
 	// To get them, convert all of the corners into heightmap 2d coordinates and get the bounding box.
 	auto PatchUVToHeightmap2DCoordinates = [&PatchToWorldOut, &FromPatchUVToPatch, &LandscapeHeightmapToWorldIn](const FVector2f& UV)
@@ -792,8 +848,8 @@ void ULandscapeTexturePatch::GetHeightShaderParams(const FTransform& LandscapeHe
 	using namespace UE::Landscape;
 
 	FTransform PatchToWorld;
-	GetCommonShaderParams(LandscapeHeightmapToWorldIn, SourceResolutionIn, DestinationResolutionIn, HeightTextureUVTransform,
-		PatchToWorld, ParamsOut.InPatchWorldDimensions, ParamsOut.InHeightmapToPatch, ParamsOut.InHeightmapTextureTransformToUV,
+	GetCommonShaderParams(LandscapeHeightmapToWorldIn, SourceResolutionIn, DestinationResolutionIn, HeightTextureUVTransform, HeightAlphaSettings.AlphaTextureUVTransform,
+		PatchToWorld, ParamsOut.InPatchWorldDimensions, ParamsOut.InHeightmapToPatchUV, ParamsOut.InHeightmapToPatchUVTransform, ParamsOut.InHeightmapAlphaToPatchUVTransform, 
 		DestinationBoundsOut, ParamsOut.InEdgeUVDeadBorder, ParamsOut.InFalloffWorldMargin);
 
 	FVector3d ComponentScale = PatchToWorld.GetScale3D();
@@ -853,7 +909,6 @@ void ULandscapeTexturePatch::GetHeightShaderParams(const FTransform& LandscapeHe
 
 	ParamsOut.InAlphaMode = static_cast<uint8>(HeightAlphaSettings.GetAlphaSourceMode());
 	ParamsOut.InAlphaTextureChannel = static_cast<uint8>(HeightAlphaSettings.AlphaTextureChannel);
-	ParamsOut.InAlphaTextureTransformToUV = static_cast<FMatrix44f>(HeightAlphaSettings.AlphaTextureUVTransform.GetTransformToUV(PatchToWorld, ParamsOut.InPatchWorldDimensions));
 
 	const ULandscapePatchEditLayer* PatchEditLayer = GetBoundEditLayer();
 	check(PatchEditLayer != nullptr);
@@ -869,8 +924,10 @@ void ULandscapeTexturePatch::GetWeightShaderParams(
 	using namespace UE::Landscape;
 
 	FTransform PatchToWorld;
-	GetCommonShaderParams(LandscapeHeightmapToWorldIn, SourceResolutionIn, DestinationResolutionIn, WeightPatchInfo->WeightTextureUVTransform,
-		PatchToWorld, ParamsOut.InPatchWorldDimensions, ParamsOut.InWeightmapToPatch, ParamsOut.InWeightmapTextureTransformToUV,
+	const FLandscapeTexturePatchTextureTransform& AlphaTextureTransform = GetWeightPatchAlphaTextureTransform(WeightPatchInfo->WeightmapLayerName);
+
+	GetCommonShaderParams(LandscapeHeightmapToWorldIn, SourceResolutionIn, DestinationResolutionIn, WeightPatchInfo->WeightTextureUVTransform, AlphaTextureTransform,
+		PatchToWorld, ParamsOut.InPatchWorldDimensions, ParamsOut.InWeightmapToPatchUV, ParamsOut.InWeightmapToPatchUVTransform, ParamsOut.InWeightmapAlphaToPatchUVTransform,
 		DestinationBoundsOut, ParamsOut.InEdgeUVDeadBorder, ParamsOut.InFalloffWorldMargin);
 
 	// Use the override blend mode if present, otherwise fall back to more general blend mode.
@@ -888,7 +945,6 @@ void ULandscapeTexturePatch::GetWeightShaderParams(
 
 	ParamsOut.InAlphaMode = static_cast<uint8>(GetWeightPatchAlphaSourceMode(WeightPatchInfo->WeightmapLayerName));
 	ParamsOut.InAlphaTextureChannel = static_cast<uint8>(GetWeightPatchAlphaTextureChannel(WeightPatchInfo->WeightmapLayerName));
-	ParamsOut.InAlphaTextureTransformToUV = static_cast<FMatrix44f>(GetWeightPatchAlphaTextureTransform(WeightPatchInfo->WeightmapLayerName).GetTransformToUV(PatchToWorld, ParamsOut.InPatchWorldDimensions));
 
 	const ULandscapePatchEditLayer* PatchEditLayer = GetBoundEditLayer();
 	check(PatchEditLayer != nullptr);
diff --git a/Engine/Plugins/Editor/LandscapePatch/Source/LandscapePatch/Private/LandscapeTexturePatchPS.h b/Engine/Plugins/Editor/LandscapePatch/Source/LandscapePatch/Private/LandscapeTexturePatchPS.h
--- a/Engine/Plugins/Editor/LandscapePatch/Source/LandscapePatch/Private/LandscapeTexturePatchPS.h
+++ b/Engine/Plugins/Editor/LandscapePatch/Source/LandscapePatch/Private/LandscapeTexturePatchPS.h
@@ -44,8 +44,9 @@ namespace UE::Landscape
 		BEGIN_SHADER_PARAMETER_STRUCT(FParameters, )
 			SHADER_PARAMETER_RDG_TEXTURE_SRV(Texture2D<float4>, InHeightPatch)
 			SHADER_PARAMETER_SAMPLER(SamplerState, InHeightPatchSampler)
-			SHADER_PARAMETER(FMatrix44f, InHeightmapToPatch)
-			SHADER_PARAMETER(FMatrix44f, InHeightmapTextureTransformToUV)
+			SHADER_PARAMETER(FMatrix44f, InHeightmapToPatchUV)
+			// InHeightmapToPatchUV matrix with the source texture transform (world space or relative) applied
+			SHADER_PARAMETER(FMatrix44f, InHeightmapToPatchUVTransform)
 			// Value in patch that corresponds to the landscape mid value, which is our "0 height".
 			SHADER_PARAMETER(float, InZeroInEncoding)
 			// Scale to apply to source values relative to the value that represents 0 height.
@@ -68,7 +69,8 @@ namespace UE::Landscape
 			SHADER_PARAMETER(uint32, InAlphaMode)
 			SHADER_PARAMETER_RDG_TEXTURE_SRV(Texture2D<float4>, InHeightPatchAlphaTexture)
 			SHADER_PARAMETER(uint32, InAlphaTextureChannel)
-			SHADER_PARAMETER(FMatrix44f, InAlphaTextureTransformToUV)
+			// InHeightmapToPatchUV matrix with the alpha source texture transform (world space or relative) applied
+			SHADER_PARAMETER(FMatrix44f, InHeightmapAlphaToPatchUVTransform)
 			// Assigned edit layer's heightmap alpha range [-1, 1]
 			SHADER_PARAMETER(float, InEditLayerHeightmapAlpha)
 
@@ -186,8 +188,9 @@ namespace UE::Landscape
 		BEGIN_SHADER_PARAMETER_STRUCT(FParameters, )
 			SHADER_PARAMETER_RDG_TEXTURE_SRV(Texture2D<float4>, InWeightPatch)
 			SHADER_PARAMETER_SAMPLER(SamplerState, InWeightPatchSampler)
-			SHADER_PARAMETER(FMatrix44f, InWeightmapToPatch)
-			SHADER_PARAMETER(FMatrix44f, InWeightmapTextureTransformToUV)
+			SHADER_PARAMETER(FMatrix44f, InWeightmapToPatchUV)
+			// InWeightmapToPatchUV matrix with the source texture transform (world space or relative) applied
+			SHADER_PARAMETER(FMatrix44f, InWeightmapToPatchUVTransform)
 			// Amount of the patch edge to not apply in UV space. Generally set to 0.5/Dimensions to avoid applying
 			// the edge half-pixels.
 			SHADER_PARAMETER(FVector2f, InEdgeUVDeadBorder)
@@ -204,7 +207,8 @@ namespace UE::Landscape
 			SHADER_PARAMETER(uint32, InAlphaMode)
 			SHADER_PARAMETER_RDG_TEXTURE_SRV(Texture2D<float4>, InWeightPatchAlphaTexture)
 			SHADER_PARAMETER(uint32, InAlphaTextureChannel)
-			SHADER_PARAMETER(FMatrix44f, InAlphaTextureTransformToUV)
+			// InWeightmapToPatchUV matrix with the alpha source texture transform (world space or relative) applied
+			SHADER_PARAMETER(FMatrix44f, InWeightmapAlphaToPatchUVTransform)
 			// Assigned edit layer's weightmap alpha range [0, 1]
 			SHADER_PARAMETER(float, InEditLayerWeightmapAlpha)
 
diff --git a/Engine/Plugins/Editor/LandscapePatch/Source/LandscapePatch/Public/LandscapeTexturePatch.h b/Engine/Plugins/Editor/LandscapePatch/Source/LandscapePatch/Public/LandscapeTexturePatch.h
--- a/Engine/Plugins/Editor/LandscapePatch/Source/LandscapePatch/Public/LandscapeTexturePatch.h
+++ b/Engine/Plugins/Editor/LandscapePatch/Source/LandscapePatch/Public/LandscapeTexturePatch.h
@@ -187,12 +187,18 @@ struct FLandscapeTexturePatchTextureTransform
 {
 	GENERATED_BODY()
 public:
-	/** If enabled, texture UVs are scaled using world-space units (1000 = tile texture every 10 meters). If disabled, scaling is applied in local UV space. */
-	UPROPERTY(EditAnywhere, BlueprintReadWrite, Category = "Texture Transform")
+	/** When enabled, texture UVs are scaled using world-space units (1000 = tile texture every 10 meters). If disabled, scaling is applied in local UV space. */
+	UPROPERTY(EditAnywhere, BlueprintReadWrite, Category = "Texture Transform", meta = (DisplayName = "Use World Unit Tiling"))
 	bool bUseWorldSpaceScale = false;
 
+	/** By default, texture UVs are centered at the patch's min bounds. Moving the patch does not change texture sampling. 
+	*  When enabled, UVs are anchored in world space at the world origin. Moving the patch's world position changes where the texture is sampled.
+	*/
+	UPROPERTY(EditAnywhere, BlueprintReadWrite, Category = "Texture Transform")
+	bool bUseWorldPositionSampling = false;
+
 	/** Scale applied to the texture's UV. By default, scaling is applied in local UV space. Relative scale 2.0 = texture is twice as large as original size. 
-	* When World Space Scale is enabled, scale 1000 = tile texture every 10 meters. 
+	*  When World Space Scale is enabled, scale 1000 = tile texture every 10 meters. 
 	*/
 	UPROPERTY(EditAnywhere, BlueprintReadWrite, Category = "Texture Transform")
 	FVector2D Scale = FVector2D(1, 1);
@@ -201,15 +207,16 @@ public:
 	UPROPERTY(EditAnywhere, BlueprintReadWrite, Category = "Texture Transform", meta = (UIMin = "0", UIMax = "360", ClampMin = "0", ClampMax = "360"))
 	float Rotation = 0;
 
-	/** Relative offset applied to the texture's UV. */
-	UPROPERTY(EditAnywhere, BlueprintReadWrite, Category = "Texture Transform", meta = (UIMin = "0", UIMax = "1", ClampMin = "0", ClampMax = "1"))
+	/** Offset applied to the sampled texture. By default, offset is in UV space [0, 1]. When bUseWorldPositionSampling is enabled, offset is in world units. */
+	UPROPERTY(EditAnywhere, BlueprintReadWrite, Category = "Texture Transform")
 	FVector2D Offset = FVector2D(0, 0);
 
-	FMatrix44d GetTransformToUV(const FTransform& PatchToWorldIn, const FVector2f& PatchWorldDimensionsIn) const;
+	FMatrix44d GetTransformAsMatrix(const FTransform& PatchToWorldIn, const FVector2f& PatchWorldDimensionsIn) const;
 
 	FORCEINLINE bool operator==(const FLandscapeTexturePatchTextureTransform InTransform) const
 	{
-		return bUseWorldSpaceScale == InTransform.bUseWorldSpaceScale && 
+		return bUseWorldSpaceScale == InTransform.bUseWorldSpaceScale &&
+			bUseWorldPositionSampling == InTransform.bUseWorldPositionSampling &&
 			Scale.X == InTransform.Scale.X && Scale.Y == InTransform.Scale.Y && 
 			Offset.X == InTransform.Offset.X && Offset.Y == InTransform.Offset.Y &&
 			Rotation == InTransform.Rotation;
@@ -1035,8 +1042,10 @@ private:
 		const FTransform& LandscapeHeightmapToWorld, bool& bHasRenderedSomething, ERHIAccess OutputAccess = ERHIAccess::None);
 
 	void GetCommonShaderParams(const FTransform& LandscapeHeightmapToWorldIn,
-		const FIntPoint& SourceResolutionIn, const FIntPoint& DestinationResolutionIn, const FLandscapeTexturePatchTextureTransform& SourceTextureUVTransformIn,
-		FTransform& PatchToWorldOut, FVector2f& PatchWorldDimensionsOut, FMatrix44f& HeightmapToPatchOut, FMatrix44f& SourceTextureTransformToUVOut,
+		const FIntPoint& SourceResolutionIn, const FIntPoint& DestinationResolutionIn,
+		const FLandscapeTexturePatchTextureTransform& SourceTextureTransformIn, const FLandscapeTexturePatchTextureTransform& AlphaTextureTransformIn,
+		FTransform& PatchToWorldOut, FVector2f& PatchWorldDimensionsOut, FMatrix44f& HeightmapToPatchUVOut, FMatrix44f& HeightmapToPatchUVTransformOut,
+		FMatrix44f& HeightmapAlphaToPatchUVTransformOut,
 		FIntRect& DestinationBoundsOut, FVector2f& EdgeUVDeadBorderOut, float& FalloffWorldMarginOut) const;
 	void GetHeightShaderParams(const FTransform& LandscapeHeightmapToWorldIn, 
 		const FIntPoint& SourceResolutionIn, const FIntPoint& DestinationResolutionIn,
